Adding compiling of Shading Model material output property for all the material compilers. Missing this was causing some material baking to fail.

#rnx

diff --git a/Engine/Source/Developer/MaterialBaking/Private/ExportMaterialProxy.h b/Engine/Source/Developer/MaterialBaking/Private/ExportMaterialProxy.h
--- a/Engine/Source/Developer/MaterialBaking/Private/ExportMaterialProxy.h
+++ b/Engine/Source/Developer/MaterialBaking/Private/ExportMaterialProxy.h
@@ -26,13 +26,11 @@ struct FExportMaterialCompiler : public FProxyMaterialCompiler
 	// gets value stored by SetMaterialProperty()
 	virtual EShaderFrequency GetCurrentShaderFrequency() const override
 	{
-		// not used by Lightmass
 		return SF_Pixel;
 	}
 
 	virtual FMaterialShadingModelField GetMaterialShadingModels() const override
 	{
-		// not used by Lightmass
 		return MSM_MAX;
 	}
 
@@ -383,6 +381,8 @@ public:
 						Compiler->Constant(0.5f)); // [-0.5,0.5] + 0.5
 				}
 				break;
+			case MP_ShadingModel:
+				return MaterialInterface->CompileProperty(&ProxyCompiler, MP_ShadingModel);
 			default:
 				return Compiler->Constant(1.0f);
 			}
@@ -399,6 +399,11 @@ public:
 			// Pass through customized UVs
 			return MaterialInterface->CompileProperty(Compiler, Property);
 		}
+		else if (Property == MP_ShadingModel)
+		{
+			return MaterialInterface->CompileProperty(Compiler, MP_ShadingModel);
+
+		}
 		else
 		{
 			return Compiler->Constant(1.0f);
diff --git a/Engine/Source/Developer/MaterialUtilities/Private/MaterialUtilities.cpp b/Engine/Source/Developer/MaterialUtilities/Private/MaterialUtilities.cpp
--- a/Engine/Source/Developer/MaterialUtilities/Private/MaterialUtilities.cpp
+++ b/Engine/Source/Developer/MaterialUtilities/Private/MaterialUtilities.cpp
@@ -564,6 +564,8 @@ public:
 							Compiler->Constant(0.5f)); // [-0.5,0.5] + 0.5
 				}
 				break;
+			case MP_ShadingModel:
+				return MaterialInterface->CompileProperty(&ProxyCompiler, MP_ShadingModel);
 			default:
 				return Compiler->Constant(1.0f);
 			}
@@ -580,6 +582,10 @@ public:
 			// Pass through customized UVs
 			return MaterialInterface->CompileProperty(Compiler, Property);
 		}
+		else if (Property == MP_ShadingModel)
+		{
+			return MaterialInterface->CompileProperty(Compiler, MP_ShadingModel);
+		}
 		else
 		{
 			return Compiler->Constant(1.0f);
diff --git a/Engine/Source/Editor/UnrealEd/Private/Lightmass/LightmassRender.cpp b/Engine/Source/Editor/UnrealEd/Private/Lightmass/LightmassRender.cpp
--- a/Engine/Source/Editor/UnrealEd/Private/Lightmass/LightmassRender.cpp
+++ b/Engine/Source/Editor/UnrealEd/Private/Lightmass/LightmassRender.cpp
@@ -436,6 +436,8 @@ public:
 					return Compiler->Lerp(Compiler->Constant3(1.0f, 1.0f, 1.0f), ColoredOpacity, Compiler->Saturate(MaterialInterface->CompileProperty(&ProxyCompiler,MP_Opacity,MFCF_ForceCast)));
 				}
 				break;
+			case MP_ShadingModel:
+				return MaterialInterface->CompileProperty(&ProxyCompiler, MP_ShadingModel);
 			default:
 				return Compiler->Constant(1.0f);
 			}
@@ -452,6 +454,10 @@ public:
 			// Pass through customized UVs
 			return MaterialInterface->CompileProperty(Compiler, Property);
 		}
+		else if (Property == MP_ShadingModel)
+		{
+			return MaterialInterface->CompileProperty(Compiler, MP_ShadingModel);
+		}
 		else
 		{
 			return Compiler->Constant(1.0f);
