Fix CIS build error.

diff --git a/Engine/Source/Developer/MaterialBaking/Public/MaterialBakingStructures.h b/Engine/Source/Developer/MaterialBaking/Public/MaterialBakingStructures.h
--- a/Engine/Source/Developer/MaterialBaking/Public/MaterialBakingStructures.h
+++ b/Engine/Source/Developer/MaterialBaking/Public/MaterialBakingStructures.h
@@ -8,6 +8,7 @@
 
 class UMaterialInterface;
 struct FRawMesh;
+struct FMeshDescription;
 
 /** Structure containing information about the material which is being baked out */
 struct FMaterialData
