Fix automation tests - crash in mesh merge

diff --git a/Engine/Source/Developer/MaterialBaking/Private/MaterialRenderItem.cpp b/Engine/Source/Developer/MaterialBaking/Private/MaterialRenderItem.cpp
--- a/Engine/Source/Developer/MaterialBaking/Private/MaterialRenderItem.cpp
+++ b/Engine/Source/Developer/MaterialBaking/Private/MaterialRenderItem.cpp
@@ -189,7 +189,7 @@ void FMeshMaterialRenderItem::PopulateWithMeshData()
 	// count number of texture coordinates for this mesh
 	const int32 NumTexcoords = [&]()
 	{
-		return FMath::Max(VertexInstanceUVs.GetNumIndices(), VertexPositionStoredUVChannel);
+		return FMath::Min(VertexInstanceUVs.GetNumIndices(), VertexPositionStoredUVChannel);
 	}();		
 
 	// check if we should use NewUVs or original UV set
