PR #5037: PerformUVBorderSmear optimization (Contributed by [REDACTED_PERSON])

diff --git a/Engine/Source/Developer/MaterialBaking/Private/MaterialBakingHelpers.cpp b/Engine/Source/Developer/MaterialBaking/Private/MaterialBakingHelpers.cpp
--- a/Engine/Source/Developer/MaterialBaking/Private/MaterialBakingHelpers.cpp
+++ b/Engine/Source/Developer/MaterialBaking/Private/MaterialBakingHelpers.cpp
@@ -4,35 +4,26 @@
 #include "Math/Color.h"
 #include "Async/ParallelFor.h"
 
-static FColor BoxBlurSample(TArray<FColor>& InBMP, int32 X, int32 Y, int32 InImageWidth, int32 InImageHeight)
+static FColor BoxBlurSample(FColor* InBMP, int32 X, int32 Y, int32 InImageWidth, int32 InImageHeight)
 {
-	const int32 SampleCount = 8;
-	const int32 PixelIndices[SampleCount] = { -(InImageWidth + 1), -InImageWidth, -(InImageWidth - 1),
-		-1, 1,
-		(InImageWidth + -1), InImageWidth, (InImageWidth + 1) };
-
-	static const int32 PixelOffsetX[SampleCount] = { -1, 0, 1,
-		-1, 1,
-		-1, 0, 1 };
-
 	int32 PixelsSampled = 0;
 	int32 CombinedColorR = 0;
 	int32 CombinedColorG = 0;
 	int32 CombinedColorB = 0;
 	int32 CombinedColorA = 0;
+	
+	const uint32 MagentaMask = FColor(255, 0, 255).DWColor();
 
-	// Take samples for blur with square indices
-	for (int32 SampleIndex = 0; SampleIndex < SampleCount; ++SampleIndex)
+	const int32 BaseIndex = (Y * InImageWidth) + X;
+	
+	int32 SampleIndex = BaseIndex - (InImageWidth + 1);
+	for (int32 YI = 0; YI < 3; ++YI)
 	{
-		const int32 PixelIndex = ((Y * InImageWidth) + X) + PixelIndices[SampleIndex];
-		const int32 XIndex = X + PixelOffsetX[SampleIndex];
-
-		// Check if we are not out of texture bounds
-		if (InBMP.IsValidIndex(PixelIndex) && XIndex >= 0 && XIndex < InImageWidth)
+		for (int32 XI = 0; XI < 3; ++XI)
 		{
-			const FColor SampledColor = InBMP[PixelIndex];
+			const FColor SampledColor = InBMP[SampleIndex++];
 			// Check if the pixel is a rendered one (not clear color)
-			if ((!(SampledColor.R == 255 && SampledColor.B == 255 && SampledColor.G == 0)))
+			if (SampledColor.DWColor() != MagentaMask)
 			{
 				CombinedColorR += SampledColor.R;
 				CombinedColorG += SampledColor.G;
@@ -41,112 +32,239 @@ static FColor BoxBlurSample(TArray<FColor>& InBMP, int32 X, int32 Y, int32 InIma
 				++PixelsSampled;
 			}
 		}
+		SampleIndex += InImageWidth - 3;
 	}
 
 	if (PixelsSampled == 0)
 	{
-		return InBMP[((Y * InImageWidth) + X)];
+		const FColor Magenta = FColor(255, 0, 255);
+		return Magenta;
 	}
 
 	return FColor(CombinedColorR / PixelsSampled, CombinedColorG / PixelsSampled, CombinedColorB / PixelsSampled, CombinedColorA / PixelsSampled);
 }
 
+static bool HasBorderingPixel(FColor* InBMP, int32 X, int32 Y, int32 InImageWidth, int32 InImageHeight)
+{
+	const uint32 MagentaMask = FColor(255, 0, 255).DWColor();
+
+	const int32 BaseIndex = Y*InImageWidth + X;
+
+	int32 SampleIndex = BaseIndex - (InImageWidth + 1);
+	for (int32 YI = 0; YI < 3; ++YI)
+	{
+		for (int32 XI = 0; XI < 3; ++XI)
+		{
+			if (InBMP[SampleIndex++].DWColor() != MagentaMask)
+			{
+				return true;
+			}
+		}
+		SampleIndex += InImageWidth - 3;
+	}
+	return false;
+}
+
 void FMaterialBakingHelpers::PerformUVBorderSmear(TArray<FColor>& InOutPixels, int32 ImageWidth, int32 ImageHeight, int32 MaxIterations)
 {
-	// This is ONLY possible because this is never called from multiple threads
-	static TArray<FColor> Swap;
-	if (Swap.Num())
+	check(InOutPixels.Num() == ImageWidth*ImageHeight);
+
+	const uint32 MagentaMask = FColor(255, 0, 255).DWColor();
+	const int32 PaddedImageWidth = ImageWidth + 2;
+	const int32 PaddedImageHeight = ImageHeight + 2;
+
+	TArray<uint32> RowsSwap1;
+	TArray<uint32> RowsSwap2;
+	TArray<uint32>* CurrentRowsLeft = &RowsSwap1;
+	TArray<uint32>* NextRowsLeft = &RowsSwap2;
+
+	//
+	// Create a copy of the image with a 1 pixel magenta border around the edge.
+	// This avoids needing to bounds check in inner loops.
+	//
+
+	FColor* Current = InOutPixels.GetData();
+	
+	TArray<FColor> ScratchBuffer;
+	ScratchBuffer.SetNumUninitialized(PaddedImageWidth*PaddedImageHeight);
+	FColor* Scratch = ScratchBuffer.GetData();
+
+	// Set top and bottom bordering pixels to magenta
+	for (int32 X = 0; X < PaddedImageWidth; ++X)
+	{
+		Scratch[X] = FColor(MagentaMask);
+		Scratch[(PaddedImageHeight - 1)*PaddedImageWidth + X] = FColor(MagentaMask);
+	}
+
+	// Set leftg and right border pixels to magenta and copy image data into rows
+	for (int32 Y = 1; Y <= ImageHeight; ++Y)
 	{
-		Swap.SetNum(0, false);
+		const int32 YOffset = Y * PaddedImageWidth;
+		Scratch[YOffset] = FColor(MagentaMask);
+		FMemory::Memcpy(&Scratch[YOffset + 1], &Current[(Y - 1)*ImageWidth], ImageWidth * sizeof(FColor));
+		Scratch[YOffset + ImageWidth + 1] = FColor(MagentaMask);
 	}
-	Swap.Append(InOutPixels);
 
-	TArray<FColor>* Current = &InOutPixels;
-	TArray<FColor>* Scratch = &Swap;
+	//
+	// Find our initial workset of all rows that have magenta pixels bordering non-magenta pixels.
+	// Also find all rows that have zero magenta pixels - these will never be added to the workset.
+	//
+
+	TArray<bool> RowCompleted;
+	RowCompleted.SetNumZeroed(PaddedImageHeight);
 
-	MaxIterations = MaxIterations < 1 ? FMath::Max(ImageWidth, ImageHeight) : MaxIterations;
-	const int32 NumThreads = [&]()
+	// Set border rows to start completed
+	RowCompleted[0] = true;
+	RowCompleted[PaddedImageHeight - 1] = true;
+
+	bool bHasAnyData = false;
+	for (int32 Y = 1; Y <= ImageHeight; ++Y)
 	{
-		return FPlatformProcess::SupportsMultithreading() ? FPlatformMisc::NumberOfCores() : 1;
-	}();
+		bool bHasMagenta = false;
+		bool bBordersNonMagenta = false;
+		const int32 YOffset = Y * PaddedImageWidth;
+		for (int32 X = 1; X <= ImageWidth; ++X)
+		{
+			if (Scratch[YOffset + X].DWColor() == MagentaMask)
+			{
+				bHasMagenta = true;
+				if (HasBorderingPixel(Scratch, X, Y, PaddedImageWidth, PaddedImageHeight))
+				{
+					bBordersNonMagenta = true;
+					break;
+				}
+			}
+		}
 
-	const int32 LinesPerThread = FMath::CeilToInt((float)ImageHeight / (float)NumThreads);
-	uint32* MagentaPixels = new uint32[NumThreads];
-	FMemory::Memset(MagentaPixels, 0xff, sizeof(uint32) * NumThreads);
+		if (!bHasMagenta)
+		{
+			bHasAnyData = true;
+			RowCompleted[Y] = true;
+		}
+		else if (bBordersNonMagenta)
+		{
+			bHasAnyData = true;
+			CurrentRowsLeft->Add(Y);
+		}
+	}
+
+	// early-out if no valid pixels were found.
+	if (!bHasAnyData)
+	{
+		FMemory::Memzero(InOutPixels.GetData(), ImageWidth * ImageHeight * sizeof(FColor));
+		return;
+	}
+
+	// Early out if all pixels were valid.
+	if (CurrentRowsLeft->Num() == 0)
+	{
+		return;
+	}
 
-	uint32 TotalPixels = ImageWidth * ImageHeight;
-	uint32 SummedMagentaPixels = 1;
+	TArray<uint32> RowRemainingPixels;
+	RowRemainingPixels.SetNumZeroed(PaddedImageHeight);
+
+	const int32 MaxThreads = FPlatformProcess::SupportsMultithreading() ? FPlatformMisc::NumberOfCores() : 1;
+
+	//
+	// Iteratively smear until all rows are filled.
+	//
 
-	// Sampling
 	int32 LoopCount = 0;
-	while (SummedMagentaPixels && (LoopCount <= MaxIterations))
+	while (CurrentRowsLeft->Num() && (MaxIterations <= 0 || LoopCount <= MaxIterations))
 	{
-		SummedMagentaPixels = 0;
-		// Left / right, Top / down
-		ParallelFor(NumThreads, [ImageWidth, ImageHeight, MaxIterations, LoopCount, Current, Scratch, &MagentaPixels, LinesPerThread]
-		(int32 Index)
+		const int32 NumThreads = FMath::Min(CurrentRowsLeft->Num(), MaxThreads);
+		const int32 LinesPerThread = FMath::CeilToInt((float)CurrentRowsLeft->Num() / (float)NumThreads);
+
+		// split up rows that still have magenta pixels amongst threads
+		ParallelFor(NumThreads, [ImageWidth, PaddedImageWidth, PaddedImageHeight, Current, Scratch, CurrentRowsLeft, LinesPerThread, &RowRemainingPixels, MagentaMask](int32 ThreadIndex)
 		{
-			if(MagentaPixels[Index] > 0)
-			{
-				MagentaPixels[Index] = 0;
+			const int32 StartY = ThreadIndex*LinesPerThread;
+			const int32 EndY = FMath::Min((ThreadIndex + 1) * LinesPerThread, CurrentRowsLeft->Num());
 
-				const int32 StartY = FMath::CeilToInt((Index)* LinesPerThread);
-				const int32 EndY = FMath::Min(FMath::CeilToInt((Index + 1) * LinesPerThread), ImageHeight);
+			for (int32 Index = StartY; Index < EndY; Index++)
+			{
+				const int32 Y = (*CurrentRowsLeft)[Index];
+				RowRemainingPixels[Index] = 0;
 
-				for (int32 Y = StartY; Y < EndY; Y++)
+				int32 PixelIndex = (Y - 1) * ImageWidth;
+				for (int32 X = 1; X <= ImageWidth; X++)
 				{
-					for (int32 X = 0; X < ImageWidth; X++)
+					FColor& Color = Current[PixelIndex++];
+					if (Color.DWColor() == MagentaMask)
 					{
-						const int32 PixelIndex = (Y * ImageWidth) + X;
-						FColor& Color = (*Current)[PixelIndex];
-						if (Color.R == 255 && Color.B == 255 && Color.G == 0)
+						const FColor SampledColor = BoxBlurSample(Scratch, X, Y, PaddedImageWidth, PaddedImageHeight);
+						// If it's a valid pixel
+						if (SampledColor.DWColor() != MagentaMask)
+						{
+							Color = SampledColor;
+						}
+						else
 						{
-							MagentaPixels[Index]++;
-							const FColor SampledColor = BoxBlurSample(*Scratch, X, Y, ImageWidth, ImageHeight);
-							// If it's a valid pixel
-							if (!(SampledColor.R == 255 && SampledColor.B == 255 && SampledColor.G == 0))
-							{
-								Color = SampledColor;
-							}
-							else
-							{
-								// If we are at the end of our iterations, replace the pixels with black
-								if (LoopCount == (MaxIterations - 1))
-								{
-									Color = FColor(0, 0, 0, 0);
-								}
-							}
+							RowRemainingPixels[Index]++;
 						}
 					}
 				}
 			}
 		});
 
-		for (int32 ThreadIndex = 0; ThreadIndex < NumThreads; ++ThreadIndex)
+		// Mark all completed rows and copy data between buffers
+		const int32 RowMemorySize = ImageWidth * sizeof(FColor);
+		for (int32 RowIndex = 0; RowIndex < CurrentRowsLeft->Num(); ++RowIndex)
 		{
-			SummedMagentaPixels += MagentaPixels[ThreadIndex];
+			const int32 Y = (*CurrentRowsLeft)[RowIndex];
+			if (!RowRemainingPixels[RowIndex])
+			{
+				RowCompleted[Y] = true;
+			}
+			FMemory::Memcpy(&Scratch[Y * PaddedImageWidth + 1], &Current[(Y - 1) * ImageWidth], RowMemorySize);
 		}
 
-		// If after iterating over all pixels we found they were all magenta, exit early
-		if(SummedMagentaPixels >= TotalPixels)
+		// Find the next set of rows to process.
+		// This will be our current rows that haven't completed, plus any bordering rows that haven't been marked completed.
+		NextRowsLeft->SetNum(0, false);
 		{
-			FMemory::Memset(Current->GetData(), 0, Current->Num() * sizeof(FColor));
-			break;
+			int32 PreviousY = -1;
+			for (int32 RowIndex = 0; RowIndex < CurrentRowsLeft->Num(); ++RowIndex)
+			{
+				const int32 Y = (*CurrentRowsLeft)[RowIndex];
+				if (!RowCompleted[Y - 1] && PreviousY < Y - 1)
+				{
+					PreviousY = Y - 1;
+					NextRowsLeft->Add(Y - 1);
+				}
+				if (!RowCompleted[Y] && RowRemainingPixels[RowIndex] && PreviousY < Y)
+				{
+					PreviousY = Y;
+					NextRowsLeft->Add(Y);
+				}
+				if (!RowCompleted[Y + 1])
+				{
+					PreviousY = Y + 1;
+					NextRowsLeft->Add(Y + 1);
+				}
+			}
 		}
 
-		TArray<FColor>& Temp = *Scratch;
-		Scratch = Current;
-		Current = &Temp;
+		// Swap rows left buffers
+		{			
+			TArray<uint32>* Temp = NextRowsLeft;
+			NextRowsLeft = CurrentRowsLeft;
+			CurrentRowsLeft = Temp;
+		}
 
 		LoopCount++;
 	}
 
-	if (Current != &InOutPixels)
+	// If we finished before replacing all pixels, replace the remaining pixels with black.
+	if (CurrentRowsLeft->Num() > 0)
 	{
-		InOutPixels.Empty();
-		InOutPixels.Append(*Current);
+		for (int32 i = 0; i < InOutPixels.Num(); ++i)
+		{
+			if (InOutPixels[i].DWColor() == MagentaMask)
+			{
+				InOutPixels[i] = FColor::Black;
+			}
+		}
 	}
-
-	delete[] MagentaPixels;
 }
-
diff --git a/Engine/Source/Developer/MaterialBaking/Private/MaterialBakingModule.cpp b/Engine/Source/Developer/MaterialBaking/Private/MaterialBakingModule.cpp
--- a/Engine/Source/Developer/MaterialBaking/Private/MaterialBakingModule.cpp
+++ b/Engine/Source/Developer/MaterialBaking/Private/MaterialBakingModule.cpp
@@ -199,12 +199,12 @@ void FMaterialBakingModule::BakeMaterials(const TArray<FMaterialData*>& Material
 					FlushRenderingCommands();
 
 					SortElement.RenderBatchArray.Empty();
-					ReadTextureOutput(RenderTargetResource, Property, CurrentOutput);
-
-					if(CurrentMaterialSettings->bPerformBorderSmear)
-					{
-						FMaterialBakingHelpers::PerformUVBorderSmear(CurrentOutput.PropertyData[Property], RenderTarget->GetSurfaceWidth(), RenderTarget->GetSurfaceHeight());
-					}
+						ReadTextureOutput(RenderTargetResource, Property, CurrentOutput);
+						
+						if (CurrentMaterialSettings->bPerformBorderSmear)
+						{
+							FMaterialBakingHelpers::PerformUVBorderSmear(CurrentOutput.PropertyData[Property], RenderTarget->GetSurfaceWidth(), RenderTarget->GetSurfaceHeight());
+						}
 #if WITH_EDITOR
 					// If saving intermediates is turned on
 					if (bSaveIntermediateTextures)
@@ -323,7 +323,7 @@ FExportMaterialProxy* FMaterialBakingModule::CreateMaterialProxy(UMaterialInterf
 
 void FMaterialBakingModule::ReadTextureOutput(FTextureRenderTargetResource* RenderTargetResource, EMaterialProperty Property, FBakeOutput& Output)
 {	
-	checkf(!Output.PropertyData.Contains(Property) && !Output.PropertySizes.Contains(Property), TEXT("Should be reading same property data twice"));
+	checkf(!Output.PropertyData.Contains(Property) && !Output.PropertySizes.Contains(Property), TEXT("Should not be reading same property data twice"));
 
 	TArray<FColor>& OutputColor = Output.PropertyData.Add(Property);
 	FIntPoint& OutputSize = Output.PropertySizes.Add(Property);	
@@ -404,10 +404,23 @@ void FMaterialBakingModule::ReadTextureOutput(FTextureRenderTargetResource* Rend
 				for (int32 PixelX = 0; PixelX < OutputSize.X; PixelX++)
 				{
 					const FFloat16Color& Pixel16 = Color16[PixelX + YOffset];
+
+					const FFloat16Color MagentaFloat16 = FFloat16Color(FLinearColor(1.0f, 0.0f, 1.0f));
 					FColor& Pixel8 = OutputColor[PixelX + YOffset];
-					Pixel8.R = (uint8)FMath::RoundToInt(Pixel16.R.GetFloat() * Scale);
-					Pixel8.G = (uint8)FMath::RoundToInt(Pixel16.G.GetFloat() * Scale);
-					Pixel8.B = (uint8)FMath::RoundToInt(Pixel16.B.GetFloat() * Scale);
+					/*if (Pixel16 == MagentaFloat16)
+					{
+						Pixel8.R = 255;
+						Pixel8.G = 0;
+						Pixel8.B = 255;
+					}
+					else*/
+					{
+						Pixel8.R = (uint8)FMath::RoundToInt(Pixel16.R.GetFloat() * Scale);
+						Pixel8.G = (uint8)FMath::RoundToInt(Pixel16.G.GetFloat() * Scale);
+						Pixel8.B = (uint8)FMath::RoundToInt(Pixel16.B.GetFloat() * Scale);						
+					}
+					
+					
 					Pixel8.A = 255;
 				}
 			}
