Fix triggered assert caused by nullptr in an auto test

diff --git a/Engine/Source/Developer/MaterialBaking/Private/MaterialRenderItemData.h b/Engine/Source/Developer/MaterialBaking/Private/MaterialRenderItemData.h
--- a/Engine/Source/Developer/MaterialBaking/Private/MaterialRenderItemData.h
+++ b/Engine/Source/Developer/MaterialBaking/Private/MaterialRenderItemData.h
@@ -17,7 +17,10 @@ public:
 		SetLightMap(InLightMap);
 		SetShadowMap(InShadowMap);
 		SetPrecomputedLightingBuffer(Buffer);
-		SetResourceCluster(InLightmapResourceCluster);
+		if (InLightmapResourceCluster)
+		{
+			SetResourceCluster(InLightmapResourceCluster);
+		}
 	}
 
 	virtual FLightInteraction GetInteraction(const class FLightSceneProxy* LightSceneProxy) const override
