Thread/core change for material baking, to include hyperthreading

diff --git a/Engine/Source/Developer/MaterialBaking/Private/MaterialBakingHelpers.cpp b/Engine/Source/Developer/MaterialBaking/Private/MaterialBakingHelpers.cpp
--- a/Engine/Source/Developer/MaterialBaking/Private/MaterialBakingHelpers.cpp
+++ b/Engine/Source/Developer/MaterialBaking/Private/MaterialBakingHelpers.cpp
@@ -164,7 +164,7 @@ void FMaterialBakingHelpers::PerformUVBorderSmear(TArray<FColor>& InOutPixels, i
 	TArray<uint32> RowRemainingPixels;
 	RowRemainingPixels.SetNumZeroed(PaddedImageHeight);
 
-	const int32 MaxThreads = FPlatformProcess::SupportsMultithreading() ? FPlatformMisc::NumberOfCores() : 1;
+	const int32 MaxThreads = FPlatformProcess::SupportsMultithreading() ? FPlatformMisc::NumberOfCoresIncludingHyperthreads() : 1;
 
 	//
 	// Iteratively smear until all rows are filled.
