MaterialBaking: Remove cached materials instances from the material proxy pool when changes are made to one of it's parent material

diff --git a/Engine/Source/Developer/MaterialBaking/Private/MaterialBakingModule.cpp b/Engine/Source/Developer/MaterialBaking/Private/MaterialBakingModule.cpp
--- a/Engine/Source/Developer/MaterialBaking/Private/MaterialBakingModule.cpp
+++ b/Engine/Source/Developer/MaterialBaking/Private/MaterialBakingModule.cpp
@@ -14,6 +14,7 @@
 #include "Framework/Application/SlateApplication.h"
 #include "MaterialBakingHelpers.h"
 #include "Async/ParallelFor.h"
+#include "Materials/MaterialInstance.h"
 
 #if WITH_EDITOR
 #include "Misc/FileHelper.h"
@@ -437,12 +438,29 @@ void FMaterialBakingModule::OnObjectModified(UObject* Object)
 	{
 		if (Object && Object->IsA<UMaterialInterface>())
 		{
-			UMaterialInterface* Material = Cast<UMaterialInterface>(Object);
-			if (MaterialProxyPool.Contains(TPair<UMaterialInterface*, EMaterialProperty>(Material, MP_BaseColor)))
+			UMaterialInterface* MaterialToInvalidate = Cast<UMaterialInterface>(Object);
+
+			// Search our proxy pool for materials or material instances that refer to MaterialToInvalidate
+			for (auto It = MaterialProxyPool.CreateIterator(); It; ++It)
 			{
-				for (int32 PropertyIndex = 0; PropertyIndex < MP_MAX; ++PropertyIndex)
+				UMaterialInterface* PoolMaterial = It.Key().Key;
+				bool bMustDelete = PoolMaterial == MaterialToInvalidate;
+
+				// No match - Test the MaterialInstance hierarchy
+				if (!bMustDelete)
+				{
+					UMaterialInstance* MaterialInstance = Cast<UMaterialInstance>(PoolMaterial);
+					while (!bMustDelete && MaterialInstance && MaterialInstance->Parent != nullptr)
+					{
+						bMustDelete = MaterialInstance->Parent == MaterialToInvalidate;
+						MaterialInstance = Cast<UMaterialInstance>(MaterialInstance->Parent);
+					}
+				}
+
+				// We have a match, remove the entry from our pool
+				if (bMustDelete)
 				{
-					MaterialProxyPool.Remove(TPair<UMaterialInterface*, EMaterialProperty>(Material, (EMaterialProperty)PropertyIndex));
+					It.RemoveCurrent();
 				}
 			}
 		}
