HLOD System - Rebuilding Medium HLODs sometimes results in bad textures
* Fixed stale materials in MaterialProxyPool
* Fixed memory leak

diff --git a/Engine/Source/Developer/MaterialBaking/Private/MaterialBakingModule.cpp b/Engine/Source/Developer/MaterialBaking/Private/MaterialBakingModule.cpp
--- a/Engine/Source/Developer/MaterialBaking/Private/MaterialBakingModule.cpp
+++ b/Engine/Source/Developer/MaterialBaking/Private/MaterialBakingModule.cpp
@@ -263,7 +263,7 @@ void FMaterialBakingModule::CleanupMaterialProxies()
 {
 	for (auto Iterator : MaterialProxyPool)
 	{
-		delete Iterator.Value;
+		delete Iterator.Value.Value;
 	}
 	MaterialProxyPool.Reset();
 }
@@ -307,16 +307,25 @@ FExportMaterialProxy* FMaterialBakingModule::CreateMaterialProxy(UMaterialInterf
 {
 	FExportMaterialProxy* Proxy = nullptr;
 
-	// Find any pooled material proxy matching this material and property
-	FExportMaterialProxy** FindResult = MaterialProxyPool.Find(TPair<UMaterialInterface*, EMaterialProperty>(Material, Property));
-	if (FindResult)
+	// Find all pooled material proxy matching this material
+	TArray<FMaterialPoolValue> Entries;
+	MaterialProxyPool.MultiFind(Material, Entries);
+
+	// Look for the matching property
+	for (FMaterialPoolValue& Entry : Entries)
 	{
-		Proxy = *FindResult;
+		if (Entry.Key == Property)
+		{
+			Proxy = Entry.Value;
+			break;
+		}
 	}
-	else
+
+	// Not found, create a new entry
+	if (Proxy == nullptr)
 	{
-		Proxy = new FExportMaterialProxy(Material, Property);		
-		MaterialProxyPool.Add(TPair<UMaterialInterface*, EMaterialProperty>(Material, Property), Proxy);
+		Proxy = new FExportMaterialProxy(Material, Property);
+		MaterialProxyPool.Add(Material, FMaterialPoolValue(Property, Proxy));
 	}
 
 	return Proxy;
@@ -443,13 +452,19 @@ void FMaterialBakingModule::OnObjectModified(UObject* Object)
 			// Search our proxy pool for materials or material instances that refer to MaterialToInvalidate
 			for (auto It = MaterialProxyPool.CreateIterator(); It; ++It)
 			{
-				UMaterialInterface* PoolMaterial = It.Key().Key;
-				bool bMustDelete = PoolMaterial == MaterialToInvalidate;
+				TWeakObjectPtr<UMaterialInterface> PoolMaterialPtr = It.Key();
+
+				// Remove stale entries from the pool
+				bool bMustDelete = PoolMaterialPtr.IsValid();
+				if (!bMustDelete)
+				{
+					bMustDelete = PoolMaterialPtr == MaterialToInvalidate;
+				}
 
 				// No match - Test the MaterialInstance hierarchy
 				if (!bMustDelete)
 				{
-					UMaterialInstance* MaterialInstance = Cast<UMaterialInstance>(PoolMaterial);
+					UMaterialInstance* MaterialInstance = Cast<UMaterialInstance>(PoolMaterialPtr);
 					while (!bMustDelete && MaterialInstance && MaterialInstance->Parent != nullptr)
 					{
 						bMustDelete = MaterialInstance->Parent == MaterialToInvalidate;
@@ -460,6 +475,9 @@ void FMaterialBakingModule::OnObjectModified(UObject* Object)
 				// We have a match, remove the entry from our pool
 				if (bMustDelete)
 				{
+					FExportMaterialProxy* Proxy = It.Value().Value;
+					delete Proxy;
+
 					It.RemoveCurrent();
 				}
 			}
diff --git a/Engine/Source/Developer/MaterialBaking/Private/MaterialBakingModule.h b/Engine/Source/Developer/MaterialBaking/Private/MaterialBakingModule.h
--- a/Engine/Source/Developer/MaterialBaking/Private/MaterialBakingModule.h
+++ b/Engine/Source/Developer/MaterialBaking/Private/MaterialBakingModule.h
@@ -53,7 +53,10 @@ private:
 	TArray<UTextureRenderTarget2D*> RenderTargetPool;
 
 	/** Pool of cached material proxies to optimize material baking workflow, stays resident when MaterialBaking.UseMaterialProxyCaching is set to 1 */
-	TMap<TPair<UMaterialInterface*, EMaterialProperty>, FExportMaterialProxy*> MaterialProxyPool;
+	typedef TWeakObjectPtr<UMaterialInterface>				FMaterialPoolKey;
+	typedef TPair<EMaterialProperty, FExportMaterialProxy*> FMaterialPoolValue;
+	typedef TMultiMap<FMaterialPoolKey, FMaterialPoolValue, FDefaultSetAllocator, TWeakObjectPtrMapKeyFuncs<FMaterialPoolKey, FMaterialPoolValue, true /*bInAllowDuplicateKeys*/>> FMaterialPoolMap;
+	FMaterialPoolMap MaterialProxyPool;
 
 	/** Pixel formats to use for baking out specific material properties */
 	EPixelFormat PerPropertyFormat[MP_MAX];
