Ensure the material proxy is always freed on the correct thread

diff --git a/Engine/Source/Developer/MaterialBaking/Private/MaterialBakingModule.cpp b/Engine/Source/Developer/MaterialBaking/Private/MaterialBakingModule.cpp
--- a/Engine/Source/Developer/MaterialBaking/Private/MaterialBakingModule.cpp
+++ b/Engine/Source/Developer/MaterialBaking/Private/MaterialBakingModule.cpp
@@ -15,6 +15,7 @@
 #include "MaterialBakingHelpers.h"
 #include "Async/ParallelFor.h"
 #include "Materials/MaterialInstance.h"
+#include "RenderingThread.h"
 
 #if WITH_EDITOR
 #include "Misc/FileHelper.h"
@@ -476,7 +477,12 @@ void FMaterialBakingModule::OnObjectModified(UObject* Object)
 				if (bMustDelete)
 				{
 					FExportMaterialProxy* Proxy = It.Value().Value;
-					delete Proxy;
+
+					ENQUEUE_RENDER_COMMAND(DeleteCachedMaterialProxy)(
+						[Proxy](FRHICommandListImmediate& RHICmdList)
+						{
+							delete Proxy;
+						});
 
 					It.RemoveCurrent();
 				}
