Certain texture maps generated by merge actors would be red instead of greyscale

#rnx

diff --git a/Engine/Source/Developer/MaterialBaking/Private/ExportMaterialProxy.h b/Engine/Source/Developer/MaterialBaking/Private/ExportMaterialProxy.h
--- a/Engine/Source/Developer/MaterialBaking/Private/ExportMaterialProxy.h
+++ b/Engine/Source/Developer/MaterialBaking/Private/ExportMaterialProxy.h
@@ -344,7 +344,7 @@ public:
 		// needs to be called in this function!!
 		Compiler->SetMaterialProperty(Property, OverrideShaderFrequency, bUsePreviousFrameTime);
 		const int32 Ret = CompilePropertyAndSetMaterialPropertyWithoutCast(Property, Compiler);
-		return Compiler->ForceCast(Ret, FMaterialAttributeDefinitionMap::GetValueType(Property), MFCF_ExactMatch);
+		return Compiler->ForceCast(Ret, FMaterialAttributeDefinitionMap::GetValueType(Property), MFCF_ExactMatch | MFCF_ReplicateValue);
 	}
 
 	/** helper for CompilePropertyAndSetMaterialProperty() */
diff --git a/Engine/Source/Developer/MaterialUtilities/Private/MaterialUtilities.cpp b/Engine/Source/Developer/MaterialUtilities/Private/MaterialUtilities.cpp
--- a/Engine/Source/Developer/MaterialUtilities/Private/MaterialUtilities.cpp
+++ b/Engine/Source/Developer/MaterialUtilities/Private/MaterialUtilities.cpp
@@ -524,7 +524,7 @@ public:
 
 		int32 Ret = CompilePropertyAndSetMaterialPropertyWithoutCast(Property, Compiler);
 
-		return Compiler->ForceCast(Ret, FMaterialAttributeDefinitionMap::GetValueType(Property), MFCF_ExactMatch);
+		return Compiler->ForceCast(Ret, FMaterialAttributeDefinitionMap::GetValueType(Property), MFCF_ExactMatch | MFCF_ReplicateValue);
 	}
 
 	/** helper for CompilePropertyAndSetMaterialProperty() */
diff --git a/Engine/Source/Runtime/Engine/Private/Materials/Material.cpp b/Engine/Source/Runtime/Engine/Private/Materials/Material.cpp
--- a/Engine/Source/Runtime/Engine/Private/Materials/Material.cpp
+++ b/Engine/Source/Runtime/Engine/Private/Materials/Material.cpp
@@ -1073,10 +1073,6 @@ void UMaterial::GetUsedTexturesAndIndices(TArray<UTexture*>& OutTextures, TArray
 
 					if (Texture)
 					{
-						if (TypeIndex == 2) 
-						{
-							UE_LOG(LogTexture, Warning, TEXT("adsf"));
-						}
 						int32 InsertIndex = OutTextures.AddUnique(Texture);
 						if (InsertIndex >= OutIndices.Num())
 						{
