Copy up //UE4/Dev-Enterprise/... to //UE4/Dev-Tools-Staging/...

diff --git a/Engine/Plugins/Enterprise/VariantManagerContent/Source/VariantManagerContent/Private/VariantObjectBinding.cpp b/Engine/Plugins/Enterprise/VariantManagerContent/Source/VariantManagerContent/Private/VariantObjectBinding.cpp
--- a/Engine/Plugins/Enterprise/VariantManagerContent/Source/VariantManagerContent/Private/VariantObjectBinding.cpp
+++ b/Engine/Plugins/Enterprise/VariantManagerContent/Source/VariantManagerContent/Private/VariantObjectBinding.cpp
@@ -2,16 +2,18 @@
 
 #include "VariantObjectBinding.h"
 
-#include "Engine/World.h"
-#include "UObject/LazyObjectPtr.h"
-#include "UObject/SoftObjectPath.h"
-#include "PropertyValue.h"
-#include "Variant.h"
 #include "LevelVariantSets.h"
 #include "LevelVariantSetsFunctionDirector.h"
-#include "GameFramework/Actor.h"
+#include "PropertyValue.h"
+#include "Variant.h"
+#include "VariantManagerObjectVersion.h"
+
 #include "Algo/Sort.h"
+#include "Engine/World.h"
 #include "FunctionCaller.h"
+#include "GameFramework/Actor.h"
+#include "UObject/LazyObjectPtr.h"
+#include "UObject/SoftObjectPath.h"
 #if WITH_EDITORONLY_DATA
 #include "K2Node_FunctionEntry.h"
 #endif
@@ -23,8 +25,10 @@ UVariantObjectBinding::UVariantObjectBinding(const FObjectInitializer& ObjectIni
 {
 }
 
-void UVariantObjectBinding::Init(UObject* InObject)
+void UVariantObjectBinding::SetObject(UObject* InObject)
 {
+	Modify();
+
 	ObjectPtr = InObject;
 	LazyObjectPtr = InObject;
 }
@@ -34,6 +38,34 @@ UVariant* UVariantObjectBinding::GetParent()
 	return Cast<UVariant>(GetOuter());
 }
 
+void UVariantObjectBinding::Serialize(FArchive& Ar)
+{
+	Super::Serialize(Ar);
+
+	Ar.UsingCustomVersion(FVariantManagerObjectVersion::GUID);
+	int32 CustomVersion = Ar.CustomVer(FVariantManagerObjectVersion::GUID);
+
+	if(Ar.IsLoading())
+	{
+		if (CustomVersion < FVariantManagerObjectVersion::StoreDisplayOrder)
+		{
+#if WITH_EDITORONLY_DATA
+			// PropertyValue and FunctionCallers won't have any display order. Assign them
+			// increasing values with FunctionCallers at the bottom
+			uint32 DisplayOrder = 0;
+			for (UPropertyValue* Property : CapturedProperties)
+			{
+				Property->SetDisplayOrder(++DisplayOrder);
+			}
+			for (FFunctionCaller& FunctionCaller : FunctionCallers)
+			{
+				FunctionCaller.SetDisplayOrder(++DisplayOrder);
+			}
+#endif
+		}
+	}
+}
+
 FText UVariantObjectBinding::GetDisplayText() const
 {
 	AActor* Actor = Cast<AActor>(GetObject());
@@ -44,10 +76,15 @@ FText UVariantObjectBinding::GetDisplayText() const
 #else
 		const FString& Label = Actor->GetName();
 #endif
-
+		CachedActorLabel = Label;
 		return FText::FromString(Label);
 	}
 
+	if (!CachedActorLabel.IsEmpty())
+	{
+		return FText::FromString(CachedActorLabel);
+	}
+
 	return FText::FromString(TEXT("<Unloaded binding>"));
 }
 
@@ -130,6 +167,14 @@ void UVariantObjectBinding::AddCapturedProperties(const TArray<UPropertyValue*>&
 		ExistingProperties.Add(Prop->GetFullDisplayString());
 	}
 
+#if WITH_EDITORONLY_DATA
+	uint32 MaxDisplayOrder = 0;
+	for (UPropertyValue* NewProp : CapturedProperties)
+	{
+		MaxDisplayOrder = FMath::Max(MaxDisplayOrder, NewProp->GetDisplayOrder());
+	}
+#endif
+
 	bool bIsMoveOperation = false;
 	TSet<UVariantObjectBinding*> ParentsModified;
 	for (UPropertyValue* NewProp : NewProperties)
@@ -148,6 +193,10 @@ void UVariantObjectBinding::AddCapturedProperties(const TArray<UPropertyValue*>&
 		NewProp->Rename(nullptr, this, REN_DontCreateRedirectors);  // Make us its Outer
 
 		CapturedProperties.Add(NewProp);
+
+#if WITH_EDITORONLY_DATA
+		NewProp->SetDisplayOrder(++MaxDisplayOrder);
+#endif
 	}
 
 	SortCapturedProperties();
@@ -172,10 +221,20 @@ void UVariantObjectBinding::RemoveCapturedProperties(const TArray<UPropertyValue
 
 void UVariantObjectBinding::SortCapturedProperties()
 {
+#if WITH_EDITORONLY_DATA
 	CapturedProperties.Sort([](const UPropertyValue& A, const UPropertyValue& B)
 	{
-		return A.GetFullDisplayString() < B.GetFullDisplayString();
+		uint32 OrderA = A.GetDisplayOrder();
+		uint32 OrderB = B.GetDisplayOrder();
+
+		if (OrderA == OrderB)
+		{
+			return A.GetFullDisplayString() < B.GetFullDisplayString();
+		}
+
+		return OrderA < OrderB;
 	});
+#endif
 }
 
 void UVariantObjectBinding::AddFunctionCallers(const TArray<FFunctionCaller>& InFunctionCallers)
diff --git a/Engine/Plugins/Enterprise/VariantManagerContent/Source/VariantManagerContent/Public/FunctionCaller.h b/Engine/Plugins/Enterprise/VariantManagerContent/Source/VariantManagerContent/Public/FunctionCaller.h
--- a/Engine/Plugins/Enterprise/VariantManagerContent/Source/VariantManagerContent/Public/FunctionCaller.h
+++ b/Engine/Plugins/Enterprise/VariantManagerContent/Source/VariantManagerContent/Public/FunctionCaller.h
@@ -18,20 +18,17 @@ struct VARIANTMANAGERCONTENT_API FFunctionCaller
 {
 	GENERATED_BODY()
 
-
 	/**
 	 * Called after this event has been serialized in order to cache the function pointer if necessary
 	 */
 	void PostSerialize(const FArchive& Ar);
 
-
 	/**
 	 * Check whether the specified function is valid for a movie scene event
 	 * Functions must have either no parameters, or a single, pass-by-value object/interface parameter, with no return parameter.
 	 */
 	static bool IsValidFunction(UFunction* Function);
 
-
 	/**
 	 * The function that should be called to invoke this event.
 	 * Functions must have either no parameters, or a single, pass-by-value object/interface parameter, with no return parameter.
@@ -42,14 +39,21 @@ struct VARIANTMANAGERCONTENT_API FFunctionCaller
 #if WITH_EDITORONLY_DATA
 
 public:
+	/**
+	 * Get the order with which the VariantManager should display this in a property list. Lower values will be shown higher up
+	 */
+	uint32 GetDisplayOrder() const;
 
+	/**
+	 * Set the order with which the VariantManager should display this in a property list. Lower values will be shown higher up
+	 */
+	void SetDisplayOrder(uint32 InDisplayOrder);
 
 	/**
 	 * Cache the function name to call from the blueprint function entry node. Will only cache the function if it has a valid signature.
 	 */
 	void CacheFunctionName();
 
-
 	/**
 	 * Check whether this event is bound to a valid blueprint entry node
 	 *
@@ -57,7 +61,6 @@ public:
 	 */
 	bool IsBoundToBlueprint() const;
 
-
 	/**
 	 * Helper function to determine whether the specified function entry is valid for this event
 	 *
@@ -66,7 +69,6 @@ public:
 	 */
 	static bool IsValidFunction(UK2Node_FunctionEntry* Node);
 
-
 	/**
 	 * Retrieve the function entry node this event is bound to
 	 *
@@ -76,7 +78,6 @@ public:
 	 */
 	UK2Node_FunctionEntry* GetFunctionEntry() const;
 
-
 	/**
 	 * Set the function entry that this event should trigger
 	 *
@@ -85,11 +86,13 @@ public:
 	void SetFunctionEntry(UK2Node_FunctionEntry* Entry);
 
 private:
-
 	/** Weak pointer to the function entry within the blueprint graph for this event. Stored as an editor-only UObject so UHT can parse it when building for non-editor. */
 	UPROPERTY()
 	TWeakObjectPtr<UObject> FunctionEntry;
 
+	UPROPERTY()
+	uint32 DisplayOrder = 0;
+
 #endif // WITH_EDITORONLY_DATA
 };
 
diff --git a/Engine/Plugins/Enterprise/VariantManagerContent/Source/VariantManagerContent/Public/PropertyValue.h b/Engine/Plugins/Enterprise/VariantManagerContent/Source/VariantManagerContent/Public/PropertyValue.h
--- a/Engine/Plugins/Enterprise/VariantManagerContent/Source/VariantManagerContent/Public/PropertyValue.h
+++ b/Engine/Plugins/Enterprise/VariantManagerContent/Source/VariantManagerContent/Public/PropertyValue.h
@@ -142,6 +142,18 @@ public:
 	FOnPropertyApplied& GetOnPropertyApplied();
 	FOnPropertyRecorded& GetOnPropertyRecorded();
 
+#if WITH_EDITORONLY_DATA
+	/**
+	 * Get the order with which the VariantManager should display this in a property list. Lower values will be shown higher up
+	 */
+	uint32 GetDisplayOrder() const;
+
+	/**
+	 * Set the order with which the VariantManager should display this in a property list. Lower values will be shown higher up
+	 */
+	void SetDisplayOrder(uint32 InDisplayOrder);
+#endif //WITH_EDITORONLY_DATA
+
 private:
 
 	void SetRecordedDataInternal(const uint8* NewDataBytes, int32 NumBytes, int32 Offset = 0);
@@ -208,6 +220,11 @@ protected:
 	FName TempName;
 	FString TempStr;
 	FText TempText;
+
+#if WITH_EDITORONLY_DATA
+	UPROPERTY()
+	uint32 DisplayOrder;
+#endif //WITH_EDITORONLY_DATA
 };
 
 // Deprecated: Only here for backwards compatibility with 4.21
diff --git a/Engine/Plugins/Enterprise/VariantManagerContent/Source/VariantManagerContent/Public/VariantObjectBinding.h b/Engine/Plugins/Enterprise/VariantManagerContent/Source/VariantManagerContent/Public/VariantObjectBinding.h
--- a/Engine/Plugins/Enterprise/VariantManagerContent/Source/VariantManagerContent/Public/VariantObjectBinding.h
+++ b/Engine/Plugins/Enterprise/VariantManagerContent/Source/VariantManagerContent/Public/VariantObjectBinding.h
@@ -18,11 +18,14 @@ class VARIANTMANAGERCONTENT_API UVariantObjectBinding : public UObject
 	GENERATED_UCLASS_BODY()
 
 public:
-
-	void Init(UObject* InObject);
+	void SetObject(UObject* InObject);
 
 	class UVariant* GetParent();
 
+	// UObject Interface
+	virtual void Serialize(FArchive& Ar) override;
+	//~ End UObject Interface
+
 	FText GetDisplayText() const;
 
 	FString GetObjectPath() const;
@@ -44,6 +47,13 @@ public:
 #endif
 
 private:
+	/**
+	 * Whenever we resolve, we cache the actor label here so that if we can't
+	 * resolve anymore we can better indicate which actor is missing, instead of just
+	 * saying 'Unloaded binding'
+	 */
+	UPROPERTY()
+	mutable FString CachedActorLabel;
 
 	UPROPERTY()
 	mutable FSoftObjectPath ObjectPtr;
diff --git a/Engine/Source/Developer/Datasmith/DatasmithCore/Private/DatasmithSceneElementsImpl.cpp b/Engine/Source/Developer/Datasmith/DatasmithCore/Private/DatasmithSceneElementsImpl.cpp
--- a/Engine/Source/Developer/Datasmith/DatasmithCore/Private/DatasmithSceneElementsImpl.cpp
+++ b/Engine/Source/Developer/Datasmith/DatasmithCore/Private/DatasmithSceneElementsImpl.cpp
@@ -204,7 +204,7 @@ FDatasmithPostProcessVolumeElementImpl::FDatasmithPostProcessVolumeElementImpl(c
 	, bEnabled( true )
 	, bUnbound( true )
 {
-	
+
 }
 
 FDatasmithCameraActorElementImpl::FDatasmithCameraActorElementImpl(const TCHAR* InName)
@@ -439,9 +439,10 @@ FDatasmithTextureElementImpl::FDatasmithTextureElementImpl(const TCHAR* InName)
 	TextureAddressY = EDatasmithTextureAddress::Wrap;
 	bAllowResize = true; // only disabled for environment maps
 	RGBCurve = -1.0;
+	ColorSpace = EDatasmithColorSpace::Default;
 
 	Data = nullptr;
-	DataSize = 0; 
+	DataSize = 0;
 }
 
 FMD5Hash FDatasmithTextureElementImpl::CalculateElementHash(bool bForce)
@@ -476,7 +477,7 @@ EDatasmithTextureMode FDatasmithTextureElementImpl::GetTextureMode() const
 	return TextureMode;
 }
 
-void FDatasmithTextureElementImpl::SetData(const uint8* InData, uint32 InDataSize, EDatasmithTextureFormat InFormat) 
+void FDatasmithTextureElementImpl::SetData(const uint8* InData, uint32 InDataSize, EDatasmithTextureFormat InFormat)
 {
 	Data = InData;
 	DataSize = InDataSize;
@@ -545,6 +546,16 @@ void FDatasmithTextureElementImpl::SetRGBCurve(float InRGBCurve)
 	RGBCurve = InRGBCurve;
 }
 
+EDatasmithColorSpace FDatasmithTextureElementImpl::GetSRGB() const
+{
+	return ColorSpace;
+}
+
+void FDatasmithTextureElementImpl::SetSRGB(EDatasmithColorSpace Option)
+{
+	ColorSpace = Option;
+}
+
 FDatasmithShaderElementImpl::FDatasmithShaderElementImpl(const TCHAR* InName)
 	: FDatasmithElementImpl( InName, EDatasmithElementType::Shader )
 	, IOR(0.0)
diff --git a/Engine/Source/Developer/Datasmith/DatasmithCore/Private/DatasmithSceneElementsImpl.h b/Engine/Source/Developer/Datasmith/DatasmithCore/Private/DatasmithSceneElementsImpl.h
--- a/Engine/Source/Developer/Datasmith/DatasmithCore/Private/DatasmithSceneElementsImpl.h
+++ b/Engine/Source/Developer/Datasmith/DatasmithCore/Private/DatasmithSceneElementsImpl.h
@@ -861,10 +861,14 @@ public:
 	virtual float GetRGBCurve() const override;
 	virtual void SetRGBCurve(float InRGBCurve) override;
 
+	virtual EDatasmithColorSpace GetSRGB() const override;
+	virtual void SetSRGB(EDatasmithColorSpace Option) override;
+
 private:
 	FString File;
 	FMD5Hash FileHash;
 	float RGBCurve;
+	EDatasmithColorSpace ColorSpace;
 	EDatasmithTextureMode TextureMode;
 	EDatasmithTextureFilter TextureFilter;
 	EDatasmithTextureAddress TextureAddressX;
diff --git a/Engine/Source/Developer/Datasmith/DatasmithCore/Private/DatasmithSceneXmlReader.cpp b/Engine/Source/Developer/Datasmith/DatasmithCore/Private/DatasmithSceneXmlReader.cpp
--- a/Engine/Source/Developer/Datasmith/DatasmithCore/Private/DatasmithSceneXmlReader.cpp
+++ b/Engine/Source/Developer/Datasmith/DatasmithCore/Private/DatasmithSceneXmlReader.cpp
@@ -192,29 +192,38 @@ void FDatasmithSceneXmlReader::ParseTextureElement(FXmlNode* InNode, TSharedPtr<
 {
 	ParseElement( InNode, OutElement.ToSharedRef() );
 
-	FString GammaAttribute = InNode->GetAttribute(TEXT("rgbcurve"));
-	if (GammaAttribute.IsEmpty() == false)
+	FString StrValue = InNode->GetAttribute(TEXT("rgbcurve"));
+	if (!StrValue.IsEmpty())
 	{
-		OutElement->SetRGBCurve(FCString::Atod(*GammaAttribute));
+		OutElement->SetRGBCurve(FCString::Atod(*StrValue));
 	}
 
-	FString StrValue = InNode->GetAttribute(TEXT("texturemode"));
-	if (StrValue != TEXT(""))
+	StrValue = InNode->GetAttribute(TEXT("srgb"));
+	if (!StrValue.IsEmpty())
+	{
+		OutElement->SetSRGB((EDatasmithColorSpace)FCString::Atoi(*StrValue));
+	}
+
+	StrValue = InNode->GetAttribute(TEXT("texturemode"));
+	if (!StrValue.IsEmpty())
 	{
 		OutElement->SetTextureMode((EDatasmithTextureMode)FCString::Atoi(*StrValue));
 	}
+
 	StrValue = InNode->GetAttribute(TEXT("texturefilter"));
-	if (StrValue != TEXT(""))
+	if (!StrValue.IsEmpty())
 	{
 		OutElement->SetTextureFilter((EDatasmithTextureFilter)FCString::Atoi(*StrValue));
 	}
+
 	StrValue = InNode->GetAttribute(TEXT("textureaddressx"));
-	if (StrValue != TEXT(""))
+	if (!StrValue.IsEmpty())
 	{
 		OutElement->SetTextureAddressX((EDatasmithTextureAddress)FCString::Atoi(*StrValue));
 	}
+
 	StrValue = InNode->GetAttribute(TEXT("textureaddressy"));
-	if (StrValue != TEXT(""))
+	if (!StrValue.IsEmpty())
 	{
 		OutElement->SetTextureAddressY((EDatasmithTextureAddress)FCString::Atoi(*StrValue));
 	}
diff --git a/Engine/Source/Developer/Datasmith/DatasmithCore/Private/DatasmithSceneXmlWriter.cpp b/Engine/Source/Developer/Datasmith/DatasmithCore/Private/DatasmithSceneXmlWriter.cpp
--- a/Engine/Source/Developer/Datasmith/DatasmithCore/Private/DatasmithSceneXmlWriter.cpp
+++ b/Engine/Source/Developer/Datasmith/DatasmithCore/Private/DatasmithSceneXmlWriter.cpp
@@ -1453,6 +1453,7 @@ void FDatasmithSceneXmlWriterImpl::WriteTextureElement(const TSharedPtr< IDatasm
 		+ FString::Printf( TEXT("\" textureaddressx=\"%i"), (int)TextureElement->GetTextureAddressX() )
 		+ FString::Printf( TEXT("\" textureaddressy=\"%i"), (int)TextureElement->GetTextureAddressY() )
 		+ FString::Printf(TEXT("\" rgbcurve=\"%f"), TextureElement->GetRGBCurve())
+		+ FString::Printf(TEXT("\" srgb=\"%i"), (int)TextureElement->GetSRGB())
 		+ TEXT("\" file=\"") + FString( TextureElement->GetFile() ) + TEXT("\">") + LINE_TERMINATOR;
 
 	SerializeToArchive( Archive, XmlString );
diff --git a/Engine/Source/Developer/Datasmith/DatasmithCore/Public/DatasmithDefinitions.h b/Engine/Source/Developer/Datasmith/DatasmithCore/Public/DatasmithDefinitions.h
--- a/Engine/Source/Developer/Datasmith/DatasmithCore/Public/DatasmithDefinitions.h
+++ b/Engine/Source/Developer/Datasmith/DatasmithCore/Public/DatasmithDefinitions.h
@@ -172,6 +172,20 @@ enum class EDatasmithTextureFormat
 	JPEG
 };
 
+/**
+ * Texture color space.
+ * Default: Leave at whatever is default for the texture mode
+ * sRGB: Enable the sRGB boolean regardless of texture mode
+ * Linear: Disable the sRGB boolean regardless of texture mode
+ */
+UENUM()
+enum class EDatasmithColorSpace : uint8
+{
+	Default,
+	sRGB,
+	Linear,
+};
+
 /**
  * Regular: lambertians, glossy materials and almost every type of material but glass, metal or highly reflective.
  * Glass: glass material, it should have appropriate index of refraction and transparency
diff --git a/Engine/Source/Developer/Datasmith/DatasmithCore/Public/IDatasmithSceneElements.h b/Engine/Source/Developer/Datasmith/DatasmithCore/Public/IDatasmithSceneElements.h
--- a/Engine/Source/Developer/Datasmith/DatasmithCore/Public/IDatasmithSceneElements.h
+++ b/Engine/Source/Developer/Datasmith/DatasmithCore/Public/IDatasmithSceneElements.h
@@ -807,6 +807,12 @@ public:
 
 	/** Set texture gamma <= 0 for auto */
 	virtual void SetRGBCurve(const float InRGBCurve) = 0;
+
+	/** Gets the color space of the texture */
+	virtual EDatasmithColorSpace GetSRGB() const = 0;
+
+	/** Sets the color space of the texture */
+	virtual void SetSRGB(EDatasmithColorSpace Option) = 0;
 };
 
 class DATASMITHCORE_API IDatasmithShaderElement : public IDatasmithElement
diff --git a/Engine/Source/Developer/MaterialBaking/Private/ExportMaterialProxy.h b/Engine/Source/Developer/MaterialBaking/Private/ExportMaterialProxy.h
--- a/Engine/Source/Developer/MaterialBaking/Private/ExportMaterialProxy.h
+++ b/Engine/Source/Developer/MaterialBaking/Private/ExportMaterialProxy.h
@@ -221,26 +221,25 @@ struct FExportMaterialCompiler : public FProxyMaterialCompiler
 class FExportMaterialProxy : public FMaterial, public FMaterialRenderProxy
 {
 public:
-	FExportMaterialProxy()
-		: FMaterial()
-	{
-		SetQualityLevelProperties(EMaterialQualityLevel::High, false, GMaxRHIFeatureLevel);
-	}
-
-	FExportMaterialProxy(UMaterialInterface* InMaterialInterface, EMaterialProperty InPropertyToCompile)
+	FExportMaterialProxy(UMaterialInterface* InMaterialInterface, EMaterialProperty InPropertyToCompile, bool bInSynchronousCompilation = true)
 		: FMaterial()
 		, MaterialInterface(InMaterialInterface)
 		, PropertyToCompile(InPropertyToCompile)
+		, bSynchronousCompilation(bInSynchronousCompilation)
 	{
 		SetQualityLevelProperties(EMaterialQualityLevel::High, false, GMaxRHIFeatureLevel);
 		Material = InMaterialInterface->GetMaterial();
 		InMaterialInterface->AppendReferencedTextures(ReferencedTextures);
-		FPlatformMisc::CreateGuid(Id);
 
 		const FMaterialResource* Resource = InMaterialInterface->GetMaterialResource(GMaxRHIFeatureLevel);
+
 		FMaterialShaderMapId ResourceId;
 		Resource->GetShaderMapId(GMaxRHIShaderPlatform, ResourceId);
 
+		// Our Id must be the same as BaseMaterialId for the shader compiler
+		// to be able to set back GameThreadShaderMap after async compilation.
+		Id = ResourceId.BaseMaterialId;
+
 		FStaticParameterSet StaticParamSet;
 		Resource->GetStaticParameterSet(GMaxRHIShaderPlatform, StaticParamSet);
 
@@ -277,7 +276,7 @@ public:
 	}
 
 	/** This override is required otherwise the shaders aren't ready for use when the surface is rendered resulting in a blank image */
-	virtual bool RequiresSynchronousCompilation() const override { return true; };
+	virtual bool RequiresSynchronousCompilation() const override { return bSynchronousCompilation; };
 
 	/**
 	* Should the shader for this material with the given platform, shader type and vertex
@@ -605,4 +604,5 @@ private:
 	/** The property to compile for rendering the sample */
 	EMaterialProperty PropertyToCompile;
 	FGuid Id;
+	bool bSynchronousCompilation;
 };
\ No newline at end of file
diff --git a/Engine/Source/Developer/MaterialBaking/Private/MaterialBakingHelpers.cpp b/Engine/Source/Developer/MaterialBaking/Private/MaterialBakingHelpers.cpp
--- a/Engine/Source/Developer/MaterialBaking/Private/MaterialBakingHelpers.cpp
+++ b/Engine/Source/Developer/MaterialBaking/Private/MaterialBakingHelpers.cpp
@@ -4,267 +4,328 @@
 #include "Math/Color.h"
 #include "Async/ParallelFor.h"
 
-static FColor BoxBlurSample(FColor* InBMP, int32 X, int32 Y, int32 InImageWidth, int32 InImageHeight)
+namespace FMaterialBakingHelpersImpl
 {
-	int32 PixelsSampled = 0;
-	int32 CombinedColorR = 0;
-	int32 CombinedColorG = 0;
-	int32 CombinedColorB = 0;
-	int32 CombinedColorA = 0;
+	static FColor BoxBlurSample(FColor* InBMP, int32 X, int32 Y, int32 InImageWidth, int32 InImageHeight)
+	{
+		int32 PixelsSampled = 0;
+		int32 CombinedColorR = 0;
+		int32 CombinedColorG = 0;
+		int32 CombinedColorB = 0;
+		int32 CombinedColorA = 0;
 	
-	const uint32 MagentaMask = FColor(255, 0, 255).DWColor();
+		const uint32 MagentaMask = FColor(255, 0, 255).DWColor();
 
-	const int32 BaseIndex = (Y * InImageWidth) + X;
+		const int32 BaseIndex = (Y * InImageWidth) + X;
 	
-	int32 SampleIndex = BaseIndex - (InImageWidth + 1);
-	for (int32 YI = 0; YI < 3; ++YI)
-	{
-		for (int32 XI = 0; XI < 3; ++XI)
+		int32 SampleIndex = BaseIndex - (InImageWidth + 1);
+		for (int32 YI = 0; YI < 3; ++YI)
 		{
-			const FColor SampledColor = InBMP[SampleIndex++];
-			// Check if the pixel is a rendered one (not clear color)
-			if (SampledColor.DWColor() != MagentaMask)
+			for (int32 XI = 0; XI < 3; ++XI)
 			{
-				CombinedColorR += SampledColor.R;
-				CombinedColorG += SampledColor.G;
-				CombinedColorB += SampledColor.B;
-				CombinedColorA += SampledColor.A;
-				++PixelsSampled;
+				const FColor SampledColor = InBMP[SampleIndex++];
+				// Check if the pixel is a rendered one (not clear color)
+				if (SampledColor.DWColor() != MagentaMask)
+				{
+					CombinedColorR += SampledColor.R;
+					CombinedColorG += SampledColor.G;
+					CombinedColorB += SampledColor.B;
+					CombinedColorA += SampledColor.A;
+					++PixelsSampled;
+				}
 			}
+			SampleIndex += InImageWidth - 3;
 		}
-		SampleIndex += InImageWidth - 3;
-	}
 
-	if (PixelsSampled == 0)
-	{
-		const FColor Magenta = FColor(255, 0, 255);
-		return Magenta;
-	}
+		if (PixelsSampled == 0)
+		{
+			const FColor Magenta = FColor(255, 0, 255);
+			return Magenta;
+		}
 
-	return FColor(CombinedColorR / PixelsSampled, CombinedColorG / PixelsSampled, CombinedColorB / PixelsSampled, CombinedColorA / PixelsSampled);
-}
+		return FColor(CombinedColorR / PixelsSampled, CombinedColorG / PixelsSampled, CombinedColorB / PixelsSampled, CombinedColorA / PixelsSampled);
+	}
 
-static bool HasBorderingPixel(FColor* InBMP, int32 X, int32 Y, int32 InImageWidth, int32 InImageHeight)
-{
-	const uint32 MagentaMask = FColor(255, 0, 255).DWColor();
+	static bool HasBorderingPixel(FColor* InBMP, int32 X, int32 Y, int32 InImageWidth, int32 InImageHeight)
+	{
+		const uint32 MagentaMask = FColor(255, 0, 255).DWColor();
 
-	const int32 BaseIndex = Y*InImageWidth + X;
+		const int32 BaseIndex = Y*InImageWidth + X;
 
-	int32 SampleIndex = BaseIndex - (InImageWidth + 1);
-	for (int32 YI = 0; YI < 3; ++YI)
-	{
-		for (int32 XI = 0; XI < 3; ++XI)
+		int32 SampleIndex = BaseIndex - (InImageWidth + 1);
+		for (int32 YI = 0; YI < 3; ++YI)
 		{
-			if (InBMP[SampleIndex++].DWColor() != MagentaMask)
+			for (int32 XI = 0; XI < 3; ++XI)
 			{
-				return true;
+				if (InBMP[SampleIndex++].DWColor() != MagentaMask)
+				{
+					return true;
+				}
 			}
+			SampleIndex += InImageWidth - 3;
 		}
-		SampleIndex += InImageWidth - 3;
+		return false;
 	}
-	return false;
-}
 
-void FMaterialBakingHelpers::PerformUVBorderSmear(TArray<FColor>& InOutPixels, int32 ImageWidth, int32 ImageHeight, int32 MaxIterations)
-{
-	check(InOutPixels.Num() == ImageWidth*ImageHeight);
+	void PerformUVBorderSmear(TArray<FColor>& InOutPixels, int32& ImageWidth, int32& ImageHeight, int32 MaxIterations, bool bCanShrink)
+	{
+		TRACE_CPUPROFILER_EVENT_SCOPE(FMaterialBakingHelpers::PerformUVBorderSmear)
 
-	const uint32 MagentaMask = FColor(255, 0, 255).DWColor();
-	const int32 PaddedImageWidth = ImageWidth + 2;
-	const int32 PaddedImageHeight = ImageHeight + 2;
+		check(InOutPixels.Num() == ImageWidth*ImageHeight);
 
-	TArray<uint32> RowsSwap1;
-	TArray<uint32> RowsSwap2;
-	TArray<uint32>* CurrentRowsLeft = &RowsSwap1;
-	TArray<uint32>* NextRowsLeft = &RowsSwap2;
+		const uint32 MagentaMask = FColor(255, 0, 255).DWColor();
+		const int32 PaddedImageWidth = ImageWidth + 2;
+		const int32 PaddedImageHeight = ImageHeight + 2;
 
-	//
-	// Create a copy of the image with a 1 pixel magenta border around the edge.
-	// This avoids needing to bounds check in inner loops.
-	//
+		TArray<uint32> RowsSwap1;
+		TArray<uint32> RowsSwap2;
+		TArray<uint32>* CurrentRowsLeft = &RowsSwap1;
+		TArray<uint32>* NextRowsLeft = &RowsSwap2;
 
-	FColor* Current = InOutPixels.GetData();
-	
-	TArray<FColor> ScratchBuffer;
-	ScratchBuffer.SetNumUninitialized(PaddedImageWidth*PaddedImageHeight);
-	FColor* Scratch = ScratchBuffer.GetData();
+		//
+		// Create a copy of the image with a 1 pixel magenta border around the edge.
+		// This avoids needing to bounds check in inner loops.
+		//
 
-	// Set top and bottom bordering pixels to magenta
-	for (int32 X = 0; X < PaddedImageWidth; ++X)
-	{
-		Scratch[X] = FColor(MagentaMask);
-		Scratch[(PaddedImageHeight - 1)*PaddedImageWidth + X] = FColor(MagentaMask);
-	}
+		FColor* Current = InOutPixels.GetData();
 
-	// Set leftg and right border pixels to magenta and copy image data into rows
-	for (int32 Y = 1; Y <= ImageHeight; ++Y)
-	{
-		const int32 YOffset = Y * PaddedImageWidth;
-		Scratch[YOffset] = FColor(MagentaMask);
-		FMemory::Memcpy(&Scratch[YOffset + 1], &Current[(Y - 1)*ImageWidth], ImageWidth * sizeof(FColor));
-		Scratch[YOffset + ImageWidth + 1] = FColor(MagentaMask);
-	}
+		// We can hopefully skip the entire smearing process if there is a single
+		// non-magenta color in the entire array since smearing that completely
+		// would lead to a monochome output.
+		if (MaxIterations == -1)
+		{
+			uint32 SingleColor     = MagentaMask;
+			bool   bHasSingleColor = true;
+			for (int32 Index = 0, Num = InOutPixels.Num(); Index < Num; ++Index)
+			{
+				const uint32 Color = Current[Index].DWColor();
+				if (Color != MagentaMask)
+				{
+					if (SingleColor == MagentaMask)
+					{
+						// This is the first time we stumble on a color other than magenta, keep it.
+						SingleColor = Color;
+					}
+					// Compare with the known non-magenta color
+					else if (SingleColor != Color)
+					{
+						bHasSingleColor = false;
+						break;
+					}
+				}
+			}
+
+			if (bHasSingleColor)
+			{
+				// Shrink to single texel when allowed
+				if (bCanShrink)
+				{
+					InOutPixels.SetNum(1);
+					ImageWidth = ImageHeight = 1;
+					Current = InOutPixels.GetData();
+				}
+
+				FColor Color(SingleColor);
+				for (int32 Index = 0, Num = InOutPixels.Num(); Index < Num; ++Index)
+				{
+					Current[Index] = Color;
+				}
 
-	//
-	// Find our initial workset of all rows that have magenta pixels bordering non-magenta pixels.
-	// Also find all rows that have zero magenta pixels - these will never be added to the workset.
-	//
+				return;
+			}
+		}
 
-	TArray<bool> RowCompleted;
-	RowCompleted.SetNumZeroed(PaddedImageHeight);
+		TArray<FColor> ScratchBuffer;
+		ScratchBuffer.SetNumUninitialized(PaddedImageWidth*PaddedImageHeight);
+		FColor* Scratch = ScratchBuffer.GetData();
 
-	// Set border rows to start completed
-	RowCompleted[0] = true;
-	RowCompleted[PaddedImageHeight - 1] = true;
+		// Set top and bottom bordering pixels to magenta
+		for (int32 X = 0; X < PaddedImageWidth; ++X)
+		{
+			Scratch[X] = FColor(MagentaMask);
+			Scratch[(PaddedImageHeight - 1)*PaddedImageWidth + X] = FColor(MagentaMask);
+		}
 
-	bool bHasAnyData = false;
-	for (int32 Y = 1; Y <= ImageHeight; ++Y)
-	{
-		bool bHasMagenta = false;
-		bool bBordersNonMagenta = false;
-		const int32 YOffset = Y * PaddedImageWidth;
-		for (int32 X = 1; X <= ImageWidth; ++X)
+		// Set leftg and right border pixels to magenta and copy image data into rows
+		for (int32 Y = 1; Y <= ImageHeight; ++Y)
 		{
-			if (Scratch[YOffset + X].DWColor() == MagentaMask)
+			const int32 YOffset = Y * PaddedImageWidth;
+			Scratch[YOffset] = FColor(MagentaMask);
+			FMemory::Memcpy(&Scratch[YOffset + 1], &Current[(Y - 1)*ImageWidth], ImageWidth * sizeof(FColor));
+			Scratch[YOffset + ImageWidth + 1] = FColor(MagentaMask);
+		}
+
+		//
+		// Find our initial workset of all rows that have magenta pixels bordering non-magenta pixels.
+		// Also find all rows that have zero magenta pixels - these will never be added to the workset.
+		//
+
+		TArray<bool> RowCompleted;
+		RowCompleted.SetNumZeroed(PaddedImageHeight);
+
+		// Set border rows to start completed
+		RowCompleted[0] = true;
+		RowCompleted[PaddedImageHeight - 1] = true;
+
+		bool bHasAnyData = false;
+		for (int32 Y = 1; Y <= ImageHeight; ++Y)
+		{
+			bool bHasMagenta = false;
+			bool bBordersNonMagenta = false;
+			const int32 YOffset = Y * PaddedImageWidth;
+			for (int32 X = 1; X <= ImageWidth; ++X)
 			{
-				bHasMagenta = true;
-				if (HasBorderingPixel(Scratch, X, Y, PaddedImageWidth, PaddedImageHeight))
+				if (Scratch[YOffset + X].DWColor() == MagentaMask)
 				{
-					bBordersNonMagenta = true;
-					break;
+					bHasMagenta = true;
+					if (HasBorderingPixel(Scratch, X, Y, PaddedImageWidth, PaddedImageHeight))
+					{
+						bBordersNonMagenta = true;
+						break;
+					}
 				}
 			}
+
+			if (!bHasMagenta)
+			{
+				bHasAnyData = true;
+				RowCompleted[Y] = true;
+			}
+			else if (bBordersNonMagenta)
+			{
+				bHasAnyData = true;
+				CurrentRowsLeft->Add(Y);
+			}
 		}
 
-		if (!bHasMagenta)
+		// early-out if no valid pixels were found.
+		if (!bHasAnyData)
 		{
-			bHasAnyData = true;
-			RowCompleted[Y] = true;
+			FMemory::Memzero(InOutPixels.GetData(), ImageWidth * ImageHeight * sizeof(FColor));
+			return;
 		}
-		else if (bBordersNonMagenta)
+
+		// Early out if all pixels were valid.
+		if (CurrentRowsLeft->Num() == 0)
 		{
-			bHasAnyData = true;
-			CurrentRowsLeft->Add(Y);
+			return;
 		}
-	}
 
-	// early-out if no valid pixels were found.
-	if (!bHasAnyData)
-	{
-		FMemory::Memzero(InOutPixels.GetData(), ImageWidth * ImageHeight * sizeof(FColor));
-		return;
-	}
-
-	// Early out if all pixels were valid.
-	if (CurrentRowsLeft->Num() == 0)
-	{
-		return;
-	}
+		TArray<uint32> RowRemainingPixels;
+		RowRemainingPixels.SetNumZeroed(PaddedImageHeight);
 
-	TArray<uint32> RowRemainingPixels;
-	RowRemainingPixels.SetNumZeroed(PaddedImageHeight);
+		const int32 MaxThreads = FPlatformProcess::SupportsMultithreading() ? FPlatformMisc::NumberOfCoresIncludingHyperthreads() : 1;
 
-	const int32 MaxThreads = FPlatformProcess::SupportsMultithreading() ? FPlatformMisc::NumberOfCoresIncludingHyperthreads() : 1;
+		//
+		// Iteratively smear until all rows are filled.
+		//
 
-	//
-	// Iteratively smear until all rows are filled.
-	//
-
-	int32 LoopCount = 0;
-	while (CurrentRowsLeft->Num() && (MaxIterations <= 0 || LoopCount <= MaxIterations))
-	{
-		const int32 NumThreads = FMath::Min(CurrentRowsLeft->Num(), MaxThreads);
-		const int32 LinesPerThread = FMath::CeilToInt((float)CurrentRowsLeft->Num() / (float)NumThreads);
-
-		// split up rows that still have magenta pixels amongst threads
-		ParallelFor(NumThreads, [ImageWidth, PaddedImageWidth, PaddedImageHeight, Current, Scratch, CurrentRowsLeft, LinesPerThread, &RowRemainingPixels, MagentaMask](int32 ThreadIndex)
+		int32 LoopCount = 0;
+		while (CurrentRowsLeft->Num() && (MaxIterations <= 0 || LoopCount <= MaxIterations))
 		{
-			const int32 StartY = ThreadIndex*LinesPerThread;
-			const int32 EndY = FMath::Min((ThreadIndex + 1) * LinesPerThread, CurrentRowsLeft->Num());
+			const int32 NumThreads = FMath::Min(CurrentRowsLeft->Num(), MaxThreads);
+			const int32 LinesPerThread = FMath::CeilToInt((float)CurrentRowsLeft->Num() / (float)NumThreads);
 
-			for (int32 Index = StartY; Index < EndY; Index++)
+			// split up rows that still have magenta pixels amongst threads
+			ParallelFor(NumThreads, [ImageWidth, PaddedImageWidth, PaddedImageHeight, Current, Scratch, CurrentRowsLeft, LinesPerThread, &RowRemainingPixels, MagentaMask](int32 ThreadIndex)
 			{
-				const int32 Y = (*CurrentRowsLeft)[Index];
-				RowRemainingPixels[Index] = 0;
+				const int32 StartY = ThreadIndex*LinesPerThread;
+				const int32 EndY = FMath::Min((ThreadIndex + 1) * LinesPerThread, CurrentRowsLeft->Num());
 
-				int32 PixelIndex = (Y - 1) * ImageWidth;
-				for (int32 X = 1; X <= ImageWidth; X++)
+				for (int32 Index = StartY; Index < EndY; Index++)
 				{
-					FColor& Color = Current[PixelIndex++];
-					if (Color.DWColor() == MagentaMask)
+					const int32 Y = (*CurrentRowsLeft)[Index];
+					RowRemainingPixels[Index] = 0;
+
+					int32 PixelIndex = (Y - 1) * ImageWidth;
+					for (int32 X = 1; X <= ImageWidth; X++)
 					{
-						const FColor SampledColor = BoxBlurSample(Scratch, X, Y, PaddedImageWidth, PaddedImageHeight);
-						// If it's a valid pixel
-						if (SampledColor.DWColor() != MagentaMask)
+						FColor& Color = Current[PixelIndex++];
+						if (Color.DWColor() == MagentaMask)
 						{
-							Color = SampledColor;
-						}
-						else
-						{
-							RowRemainingPixels[Index]++;
+							const FColor SampledColor = BoxBlurSample(Scratch, X, Y, PaddedImageWidth, PaddedImageHeight);
+							// If it's a valid pixel
+							if (SampledColor.DWColor() != MagentaMask)
+							{
+								Color = SampledColor;
+							}
+							else
+							{
+								RowRemainingPixels[Index]++;
+							}
 						}
 					}
 				}
-			}
-		});
-
-		// Mark all completed rows and copy data between buffers
-		const int32 RowMemorySize = ImageWidth * sizeof(FColor);
-		for (int32 RowIndex = 0; RowIndex < CurrentRowsLeft->Num(); ++RowIndex)
-		{
-			const int32 Y = (*CurrentRowsLeft)[RowIndex];
-			if (!RowRemainingPixels[RowIndex])
-			{
-				RowCompleted[Y] = true;
-			}
-			FMemory::Memcpy(&Scratch[Y * PaddedImageWidth + 1], &Current[(Y - 1) * ImageWidth], RowMemorySize);
-		}
+			});
 
-		// Find the next set of rows to process.
-		// This will be our current rows that haven't completed, plus any bordering rows that haven't been marked completed.
-		NextRowsLeft->SetNum(0, false);
-		{
-			int32 PreviousY = -1;
+			// Mark all completed rows and copy data between buffers
+			const int32 RowMemorySize = ImageWidth * sizeof(FColor);
 			for (int32 RowIndex = 0; RowIndex < CurrentRowsLeft->Num(); ++RowIndex)
 			{
 				const int32 Y = (*CurrentRowsLeft)[RowIndex];
-				if (!RowCompleted[Y - 1] && PreviousY < Y - 1)
-				{
-					PreviousY = Y - 1;
-					NextRowsLeft->Add(Y - 1);
-				}
-				if (!RowCompleted[Y] && RowRemainingPixels[RowIndex] && PreviousY < Y)
+				if (!RowRemainingPixels[RowIndex])
 				{
-					PreviousY = Y;
-					NextRowsLeft->Add(Y);
+					RowCompleted[Y] = true;
 				}
-				if (!RowCompleted[Y + 1])
+				FMemory::Memcpy(&Scratch[Y * PaddedImageWidth + 1], &Current[(Y - 1) * ImageWidth], RowMemorySize);
+			}
+
+			// Find the next set of rows to process.
+			// This will be our current rows that haven't completed, plus any bordering rows that haven't been marked completed.
+			NextRowsLeft->SetNum(0, false);
+			{
+				int32 PreviousY = -1;
+				for (int32 RowIndex = 0; RowIndex < CurrentRowsLeft->Num(); ++RowIndex)
 				{
-					PreviousY = Y + 1;
-					NextRowsLeft->Add(Y + 1);
+					const int32 Y = (*CurrentRowsLeft)[RowIndex];
+					if (!RowCompleted[Y - 1] && PreviousY < Y - 1)
+					{
+						PreviousY = Y - 1;
+						NextRowsLeft->Add(Y - 1);
+					}
+					if (!RowCompleted[Y] && RowRemainingPixels[RowIndex] && PreviousY < Y)
+					{
+						PreviousY = Y;
+						NextRowsLeft->Add(Y);
+					}
+					if (!RowCompleted[Y + 1])
+					{
+						PreviousY = Y + 1;
+						NextRowsLeft->Add(Y + 1);
+					}
 				}
 			}
-		}
 
-		// Swap rows left buffers
-		{			
-			TArray<uint32>* Temp = NextRowsLeft;
-			NextRowsLeft = CurrentRowsLeft;
-			CurrentRowsLeft = Temp;
-		}
+			// Swap rows left buffers
+			{			
+				TArray<uint32>* Temp = NextRowsLeft;
+				NextRowsLeft = CurrentRowsLeft;
+				CurrentRowsLeft = Temp;
+			}
 
-		LoopCount++;
-	}
+			LoopCount++;
+		}
 
-	// If we finished before replacing all pixels, replace the remaining pixels with black.
-	if (CurrentRowsLeft->Num() > 0)
-	{
-		for (int32 i = 0; i < InOutPixels.Num(); ++i)
+		// If we finished before replacing all pixels, replace the remaining pixels with black.
+		if (CurrentRowsLeft->Num() > 0)
 		{
-			if (InOutPixels[i].DWColor() == MagentaMask)
+			for (int32 i = 0; i < InOutPixels.Num(); ++i)
 			{
-				InOutPixels[i] = FColor::Black;
+				if (InOutPixels[i].DWColor() == MagentaMask)
+				{
+					InOutPixels[i] = FColor::Black;
+				}
 			}
 		}
 	}
 }
+
+void FMaterialBakingHelpers::PerformUVBorderSmear(TArray<FColor>& InOutPixels, int32 ImageWidth, int32 ImageHeight, int32 MaxIterations)
+{
+	FMaterialBakingHelpersImpl::PerformUVBorderSmear(InOutPixels, ImageWidth, ImageHeight, MaxIterations, false);
+}
+
+void FMaterialBakingHelpers::PerformUVBorderSmearAndShrink(TArray<FColor>& InOutPixels, int32& InOutImageWidth, int32& InOutImageHeight)
+{
+	FMaterialBakingHelpersImpl::PerformUVBorderSmear(InOutPixels, InOutImageWidth, InOutImageHeight, -1, true);
+}
\ No newline at end of file
diff --git a/Engine/Source/Developer/MaterialBaking/Private/MaterialBakingModule.cpp b/Engine/Source/Developer/MaterialBaking/Private/MaterialBakingModule.cpp
--- a/Engine/Source/Developer/MaterialBaking/Private/MaterialBakingModule.cpp
+++ b/Engine/Source/Developer/MaterialBaking/Private/MaterialBakingModule.cpp
@@ -13,10 +13,12 @@
 #include "MaterialBakingStructures.h"
 #include "Framework/Application/SlateApplication.h"
 #include "MaterialBakingHelpers.h"
+#include "Async/Async.h"
 #include "Async/ParallelFor.h"
 #include "Materials/MaterialInstance.h"
 #include "RenderingThread.h"
-
+#include "Misc/ScopedSlowTask.h"
+#include "MeshDescription.h"
 #if WITH_EDITOR
 #include "Misc/FileHelper.h"
 #endif
@@ -42,6 +44,76 @@ static TAutoConsoleVariable<int32> CVarSaveIntermediateTextures(
 	TEXT("1: Turned On"),
 	ECVF_Default);
 
+namespace FMaterialBakingModuleImpl
+{
+	// Custom dynamic mesh allocator specifically tailored for Material Baking.
+	// This will always reuse the same couple buffers, so searching linearly is not a problem.
+	class FMaterialBakingDynamicMeshBufferAllocator : public FDynamicMeshBufferAllocator
+	{
+		TArray<FIndexBufferRHIRef>  IndexBuffers;
+		TArray<FVertexBufferRHIRef> VertexBuffers;
+
+		template <typename RefType>
+		RefType GetSmallestFit(int32 SizeInBytes, TArray<RefType>& Array)
+		{
+			int32 SmallestFitIndex = -1;
+			int32 SmallestFitSize  = -1;
+			for (int32 Index = 0; Index < Array.Num(); ++Index)
+			{
+				int32 Size = Array[Index]->GetSize();
+				if (Size >= SizeInBytes && (SmallestFitIndex == -1 || Size < SmallestFitSize))
+				{
+					SmallestFitIndex = Index;
+					SmallestFitSize  = Size;
+				}
+			}
+
+			RefType Ref;
+			if (SmallestFitIndex != -1)
+			{
+				Ref = Array[SmallestFitIndex];
+				Array.RemoveAtSwap(SmallestFitIndex);
+			}
+
+			return Ref;
+		}
+
+		virtual FIndexBufferRHIRef AllocIndexBuffer(uint32 NumElements) override
+		{
+			TRACE_CPUPROFILER_EVENT_SCOPE(AllocIndexBuffer)
+			FIndexBufferRHIRef Ref = GetSmallestFit(GetIndexBufferSize(NumElements), IndexBuffers);
+			if (Ref.IsValid())
+			{
+				return Ref;
+			}
+
+			return FDynamicMeshBufferAllocator::AllocIndexBuffer(NumElements);
+		}
+
+		virtual void ReleaseIndexBuffer(FIndexBufferRHIRef& IndexBufferRHI) override
+		{
+			IndexBuffers.Add(MoveTemp(IndexBufferRHI));
+		}
+
+		virtual FVertexBufferRHIRef AllocVertexBuffer(uint32 Stride, uint32 NumElements) override
+		{
+			TRACE_CPUPROFILER_EVENT_SCOPE(AllocVertexBuffer)
+			FVertexBufferRHIRef Ref = GetSmallestFit(GetVertexBufferSize(Stride, NumElements), VertexBuffers);
+			if (Ref.IsValid())
+			{
+				return Ref;
+			}
+
+			return FDynamicMeshBufferAllocator::AllocVertexBuffer(Stride, NumElements);
+		}
+
+		virtual void ReleaseVertexBuffer(FVertexBufferRHIRef& VertexBufferRHI) override
+		{
+			VertexBuffers.Add(MoveTemp(VertexBufferRHI));
+		}
+	};
+}
+
 void FMaterialBakingModule::StartupModule()
 {
 	// Set which properties should enforce gamme correction
@@ -73,7 +145,7 @@ void FMaterialBakingModule::StartupModule()
 
 void FMaterialBakingModule::ShutdownModule()
 {
-	// Urnegister customization and callback
+	// Unregister customization and callback
 	FPropertyEditorModule* PropertyEditorModule = FModuleManager::GetModulePtr<FPropertyEditorModule>("PropertyEditor");
 
 	if (PropertyEditorModule)
@@ -84,78 +156,145 @@ void FMaterialBakingModule::ShutdownModule()
 }
 
 void FMaterialBakingModule::BakeMaterials(const TArray<FMaterialData*>& MaterialSettings, const TArray<FMeshData*>& MeshSettings, TArray<FBakeOutput>& Output)
-{	
+{
+	TRACE_CPUPROFILER_EVENT_SCOPE(FMaterialBakingModule::BakeMaterials)
+
 	checkf(MaterialSettings.Num() == MeshSettings.Num(), TEXT("Number of material settings does not match that of MeshSettings"));
 	const int32 NumMaterials = MaterialSettings.Num();
 	const bool bSaveIntermediateTextures = CVarSaveIntermediateTextures.GetValueOnAnyThread() == 1;
 
-	for (int32 MaterialIndex = 0; MaterialIndex < NumMaterials; ++MaterialIndex)
+	FMaterialBakingModuleImpl::FMaterialBakingDynamicMeshBufferAllocator MaterialBakingDynamicMeshBufferAllocator;
+
+	FScopedSlowTask Progress(NumMaterials, LOCTEXT("BakeMaterials", "Baking Materials..."), true );
+	Progress.MakeDialog(true);
+
+	TArray<uint32> ProcessingOrder;
+	ProcessingOrder.Reserve(MeshSettings.Num());
+	for (int32 Index = 0; Index < MeshSettings.Num(); ++Index)
+	{
+		ProcessingOrder.Add(Index);
+	}
+
+	// Start with the biggest mesh first so we can always reuse the same vertex/index buffers.
+	// This will decrease the number of allocations backed by newly allocated memory from the OS,
+	// which will reduce soft page faults while copying into that memory.
+	// Soft page faults are now incredibly expensive on Windows 10.
+	Algo::SortBy(
+		ProcessingOrder,
+		[&MeshSettings](const uint32 Index){ return MeshSettings[Index]->RawMeshDescription ? MeshSettings[Index]->RawMeshDescription->Vertices().Num() : 0; },
+		TGreater<>()
+	);
+
+	TAtomic<uint32> NumTasks(0);
+	Output.SetNum(NumMaterials);
+
+	// Create all material proxies right away to start compiling shaders asynchronously and avoid stalling the baking process as much as possible
+	{
+		TRACE_CPUPROFILER_EVENT_SCOPE(CreateMaterialProxies)
+
+		for (int32 Index = 0; Index < NumMaterials; ++Index)
+		{
+			int32 MaterialIndex = ProcessingOrder[Index];
+			const FMaterialData* CurrentMaterialSettings = MaterialSettings[MaterialIndex];
+
+			TArray<UTexture*> MaterialTextures;
+			CurrentMaterialSettings->Material->GetUsedTextures(MaterialTextures, EMaterialQualityLevel::Num, true, GMaxRHIFeatureLevel, true);
+
+			// Force load materials used by the current material
+			{
+				TRACE_CPUPROFILER_EVENT_SCOPE(LoadTexturesForMaterial)
+
+				for (UTexture* Texture : MaterialTextures)
+				{
+					if (Texture != NULL)
+					{
+						UTexture2D* Texture2D = Cast<UTexture2D>(Texture);
+						if (Texture2D)
+						{
+							Texture2D->SetForceMipLevelsToBeResident(30.0f);
+							Texture2D->WaitForStreaming();
+						}
+					}
+				}
+			}
+
+			for (TMap<EMaterialProperty, FIntPoint>::TConstIterator PropertySizeIterator = CurrentMaterialSettings->PropertySizes.CreateConstIterator(); PropertySizeIterator; ++PropertySizeIterator)
+			{
+				// They will be stored in the pool and compiled asynchronously
+				CreateMaterialProxy(CurrentMaterialSettings->Material, PropertySizeIterator.Key());
+			}
+		}
+	}
+
+	for (int32 Index = 0; Index < NumMaterials; ++Index)
 	{
+		TRACE_CPUPROFILER_EVENT_SCOPE(BakeOneMaterial)
+
+		Progress.EnterProgressFrame();
+
+		int32 MaterialIndex = ProcessingOrder[Index];
 		const FMaterialData* CurrentMaterialSettings = MaterialSettings[MaterialIndex];
 		const FMeshData* CurrentMeshSettings = MeshSettings[MaterialIndex];
-		FBakeOutput& CurrentOutput = [&Output]() -> FBakeOutput&
-		{
-			Output.AddDefaulted(1);
-			return Output.Last();
-		}();
+		FBakeOutput& CurrentOutput = Output[MaterialIndex];
 
 		// Canvas per size / special property?
 		TArray<TPair<UTextureRenderTarget2D*, FSceneViewFamily>> TargetsViewFamilyPairs;
 		TArray<FExportMaterialProxy*> MaterialRenderProxies;
 		TArray<EMaterialProperty> MaterialPropertiesToBakeOut;
 
-		TArray<UTexture*> MaterialTextures;
-		CurrentMaterialSettings->Material->GetUsedTextures(MaterialTextures, EMaterialQualityLevel::Num, true, GMaxRHIFeatureLevel, true);
-
-		// Force load materials used by the current material
-		for (UTexture* Texture : MaterialTextures)
 		{
-			if (Texture != NULL)
+			TRACE_CPUPROFILER_EVENT_SCOPE(CreateMaterialProxies)
+
+			for (TMap<EMaterialProperty, FIntPoint>::TConstIterator PropertySizeIterator = CurrentMaterialSettings->PropertySizes.CreateConstIterator(); PropertySizeIterator; ++PropertySizeIterator)
 			{
-				UTexture2D* Texture2D = Cast<UTexture2D>(Texture);
-				if (Texture2D)
-				{
-					Texture2D->SetForceMipLevelsToBeResident(30.0f);
-					Texture2D->WaitForStreaming();
-				}
+				const EMaterialProperty& Property = PropertySizeIterator.Key();
+				UTextureRenderTarget2D* RenderTarget = CreateRenderTarget(PerPropertyGamma[Property], PerPropertyFormat[Property], PropertySizeIterator.Value());
+				FExportMaterialProxy* Proxy = CreateMaterialProxy(CurrentMaterialSettings->Material, Property);
+
+				FSceneViewFamily ViewFamily(FSceneViewFamily::ConstructionValues(RenderTarget->GameThread_GetRenderTargetResource(), nullptr,
+					FEngineShowFlags(ESFIM_Game))
+					.SetWorldTimes(0.0f, 0.0f, 0.0f)
+					.SetGammaCorrection(RenderTarget->GameThread_GetRenderTargetResource()->GetDisplayGamma()));
+
+				TargetsViewFamilyPairs.Add(TPair<UTextureRenderTarget2D*, FSceneViewFamily>(RenderTarget, Forward<FSceneViewFamily>(ViewFamily)));
+				MaterialRenderProxies.Add(Proxy);
+				MaterialPropertiesToBakeOut.Add(Property);
 			}
 		}
-		
-		// Generate view family, material proxy and render targets for each material property being baked out
-		for (TMap<EMaterialProperty, FIntPoint>::TConstIterator PropertySizeIterator = CurrentMaterialSettings->PropertySizes.CreateConstIterator(); PropertySizeIterator; ++PropertySizeIterator)
-		{
-			const EMaterialProperty& Property = PropertySizeIterator.Key();
-			UTextureRenderTarget2D* RenderTarget = CreateRenderTarget(PerPropertyGamma[Property], PerPropertyFormat[Property], PropertySizeIterator.Value());
-			FExportMaterialProxy* Proxy = CreateMaterialProxy(CurrentMaterialSettings->Material, Property);
-
-			FSceneViewFamily ViewFamily(FSceneViewFamily::ConstructionValues(RenderTarget->GameThread_GetRenderTargetResource(), nullptr,
-				FEngineShowFlags(ESFIM_Game))
-				.SetWorldTimes(0.0f, 0.0f, 0.0f)
-				.SetGammaCorrection(RenderTarget->GameThread_GetRenderTargetResource()->GetDisplayGamma()));
-
-			TargetsViewFamilyPairs.Add(TPair<UTextureRenderTarget2D*, FSceneViewFamily>(RenderTarget, Forward<FSceneViewFamily>(ViewFamily)));
-			MaterialRenderProxies.Add(Proxy);
-			MaterialPropertiesToBakeOut.Add(Property);
-		}
 
 		const int32 NumPropertiesToRender = CurrentMaterialSettings->PropertySizes.Num();
 		if (NumPropertiesToRender > 0)
 		{
+			TRACE_CPUPROFILER_EVENT_SCOPE(RenderOneMaterial)
+
+			// Ensure data in memory will not change place while adding to avoid race conditions
+			CurrentOutput.PropertyData.Reserve(NumPropertiesToRender);
+			CurrentOutput.PropertySizes.Reserve(NumPropertiesToRender);
+
 			FCanvas Canvas(TargetsViewFamilyPairs[0].Key->GameThread_GetRenderTargetResource(), nullptr, FApp::GetCurrentTime() - GStartTime, FApp::GetDeltaTime(), FApp::GetCurrentTime() - GStartTime, GMaxRHIFeatureLevel);
 			Canvas.SetAllowedModes(FCanvas::Allow_Flush);
 
 			UTextureRenderTarget2D* PreviousRenderTarget = nullptr;
 			FTextureRenderTargetResource* RenderTargetResource = nullptr;
 
-			FMeshMaterialRenderItem RenderItem(CurrentMaterialSettings, CurrentMeshSettings, MaterialPropertiesToBakeOut[0]);
+			// The RenderItem will use our custom allocator for vertex/index buffers.
+			FMeshMaterialRenderItem RenderItem(CurrentMaterialSettings, CurrentMeshSettings, MaterialPropertiesToBakeOut[0], &MaterialBakingDynamicMeshBufferAllocator);
 			FCanvas::FCanvasSortElement& SortElement = Canvas.GetSortElement(Canvas.TopDepthSortKey());
 			
 			for (int32 PropertyIndex = 0; PropertyIndex < NumPropertiesToRender; ++PropertyIndex)
 			{
 				const EMaterialProperty Property = MaterialPropertiesToBakeOut[PropertyIndex];
+				FExportMaterialProxy* ExportMaterialProxy = MaterialRenderProxies[PropertyIndex];
+
+				if (!ExportMaterialProxy->IsCompilationFinished())
+				{
+					TRACE_CPUPROFILER_EVENT_SCOPE(WaitForMaterialProxyCompilation)
+					ExportMaterialProxy->FinishCompilation();
+				}
+
 				// Update render item
 				RenderItem.MaterialProperty = Property;
-				RenderItem.MaterialRenderProxy = MaterialRenderProxies[PropertyIndex];
+				RenderItem.MaterialRenderProxy = ExportMaterialProxy;
 
 				UTextureRenderTarget2D* RenderTarget = TargetsViewFamilyPairs[PropertyIndex].Key;
 				if (RenderTarget != nullptr)
@@ -198,34 +337,63 @@ void FMaterialBakingModule::BakeMaterials(const TArray<FMaterialData*>& Material
 
 					// Do rendering
 					Canvas.Flush_GameThread();
-					FlushRenderingCommands();
 
 					SortElement.RenderBatchArray.Empty();
-						ReadTextureOutput(RenderTargetResource, Property, CurrentOutput);
-						
-						if (CurrentMaterialSettings->bPerformBorderSmear)
+
+					// Reading texture does an implicit FlushRenderingCommands()
+					ReadTextureOutput(RenderTargetResource, Property, CurrentOutput);
+
+					// Schedule final processing for another thread
+					NumTasks++;
+					auto FinalProcessing =
+						[&NumTasks, bSaveIntermediateTextures, CurrentMaterialSettings, MaterialIndex, &Output, Property]()
 						{
-							FMaterialBakingHelpers::PerformUVBorderSmear(CurrentOutput.PropertyData[Property], RenderTarget->GetSurfaceWidth(), RenderTarget->GetSurfaceHeight());
-						}
+							FBakeOutput& CurrentOutput = Output[MaterialIndex];
+
+							if (CurrentMaterialSettings->bPerformBorderSmear)
+							{
+								// This will resize the output to a single pixel if the result is monochrome.
+								FMaterialBakingHelpers::PerformUVBorderSmearAndShrink(
+									CurrentOutput.PropertyData[Property],
+									CurrentOutput.PropertySizes[Property].X,
+									CurrentOutput.PropertySizes[Property].Y
+								);
+							}
 #if WITH_EDITOR
-					// If saving intermediates is turned on
-					if (bSaveIntermediateTextures)
+							// If saving intermediates is turned on
+							if (bSaveIntermediateTextures)
+							{
+								const UEnum* PropertyEnum = StaticEnum<EMaterialProperty>();
+								FName PropertyName = PropertyEnum->GetNameByValue(Property);
+								FString TrimmedPropertyName = PropertyName.ToString();
+								TrimmedPropertyName.RemoveFromStart(TEXT("MP_"));
+
+								const FString DirectoryPath = FPaths::ConvertRelativePathToFull(FPaths::ProjectIntermediateDir() + TEXT("MaterialBaking/"));
+								FString FilenameString = FString::Printf(TEXT("%s%s-%d-%s.bmp"), *DirectoryPath, *CurrentMaterialSettings->Material->GetName(), MaterialIndex, *TrimmedPropertyName);
+								FFileHelper::CreateBitmap(*FilenameString, CurrentOutput.PropertySizes[Property].X, CurrentOutput.PropertySizes[Property].Y, CurrentOutput.PropertyData[Property].GetData());
+							}
+#endif // WITH_EDITOR
+							NumTasks--;
+						};
+
+					if (FPlatformProcess::SupportsMultithreading())
 					{
-						const UEnum* PropertyEnum = StaticEnum<EMaterialProperty>();
-						FName PropertyName = PropertyEnum->GetNameByValue(Property);
-						FString TrimmedPropertyName = PropertyName.ToString();
-						TrimmedPropertyName.RemoveFromStart(TEXT("MP_"));
-
-						const FString DirectoryPath = FPaths::ConvertRelativePathToFull(FPaths::ProjectIntermediateDir() + TEXT("MaterialBaking/"));
-						FString FilenameString = FString::Printf(TEXT("%s%s-%d-%s.bmp"), *DirectoryPath, *CurrentMaterialSettings->Material->GetName(), MaterialIndex, *TrimmedPropertyName);
-						FFileHelper::CreateBitmap(*FilenameString, CurrentOutput.PropertySizes[Property].X, CurrentOutput.PropertySizes[Property].Y, CurrentOutput.PropertyData[Property].GetData());
+						Async(EAsyncExecution::ThreadPool, MoveTemp(FinalProcessing));
+					}
+					else
+					{
+						FinalProcessing();
 					}
 				}
-#endif // WITH_EDITOR
 			}
 		}
 	}
 
+	while (NumTasks.Load(EMemoryOrder::Relaxed) > 0)
+	{
+		FPlatformProcess::Sleep(0.1f);
+	}
+
 	if (!CVarUseMaterialProxyCaching.GetValueOnAnyThread())
 	{
 		CleanupMaterialProxies();
@@ -271,6 +439,8 @@ void FMaterialBakingModule::CleanupMaterialProxies()
 
 UTextureRenderTarget2D* FMaterialBakingModule::CreateRenderTarget(bool bInForceLinearGamma, EPixelFormat InPixelFormat, const FIntPoint& InTargetSize)
 {
+	TRACE_CPUPROFILER_EVENT_SCOPE(FMaterialBakingModule::CreateRenderTarget)
+
 	UTextureRenderTarget2D* RenderTarget = nullptr;
 	const int32 MaxTextureSize = 1 << (MAX_TEXTURE_MIP_COUNT - 1); // Don't use GetMax2DTextureDimension() as this is for the RHI only.
 	const FIntPoint ClampedTargetSize(FMath::Clamp(InTargetSize.X, 1, MaxTextureSize), FMath::Clamp(InTargetSize.Y, 1, MaxTextureSize));
@@ -288,6 +458,8 @@ UTextureRenderTarget2D* FMaterialBakingModule::CreateRenderTarget(bool bInForceL
 	}
 	else
 	{
+		TRACE_CPUPROFILER_EVENT_SCOPE(CreateNewRenderTarget)
+
 		// Not found - create a new one.
 		RenderTarget = NewObject<UTextureRenderTarget2D>();
 		check(RenderTarget);
@@ -306,6 +478,8 @@ UTextureRenderTarget2D* FMaterialBakingModule::CreateRenderTarget(bool bInForceL
 
 FExportMaterialProxy* FMaterialBakingModule::CreateMaterialProxy(UMaterialInterface* Material, const EMaterialProperty Property)
 {
+	TRACE_CPUPROFILER_EVENT_SCOPE(FMaterialBakingModule::CreateMaterialProxy)
+
 	FExportMaterialProxy* Proxy = nullptr;
 
 	// Find all pooled material proxy matching this material
@@ -325,7 +499,7 @@ FExportMaterialProxy* FMaterialBakingModule::CreateMaterialProxy(UMaterialInterf
 	// Not found, create a new entry
 	if (Proxy == nullptr)
 	{
-		Proxy = new FExportMaterialProxy(Material, Property);
+		Proxy = new FExportMaterialProxy(Material, Property, false /* bInSynchronousCompilation */);
 		MaterialProxyPool.Add(Material, FMaterialPoolValue(Property, Proxy));
 	}
 
@@ -333,7 +507,9 @@ FExportMaterialProxy* FMaterialBakingModule::CreateMaterialProxy(UMaterialInterf
 }
 
 void FMaterialBakingModule::ReadTextureOutput(FTextureRenderTargetResource* RenderTargetResource, EMaterialProperty Property, FBakeOutput& Output)
-{	
+{
+	TRACE_CPUPROFILER_EVENT_SCOPE(FMaterialBakingModule::ReadTextureOutput)
+
 	checkf(!Output.PropertyData.Contains(Property) && !Output.PropertySizes.Contains(Property), TEXT("Should not be reading same property data twice"));
 
 	TArray<FColor>& OutputColor = Output.PropertyData.Add(Property);
@@ -444,6 +620,8 @@ void FMaterialBakingModule::ReadTextureOutput(FTextureRenderTargetResource* Rend
 
 void FMaterialBakingModule::OnObjectModified(UObject* Object)
 {
+	TRACE_CPUPROFILER_EVENT_SCOPE(FMaterialBakingModule::OnObjectModified)
+
 	if (CVarUseMaterialProxyCaching.GetValueOnAnyThread())
 	{
 		if (Object && Object->IsA<UMaterialInterface>())
diff --git a/Engine/Source/Developer/MaterialBaking/Private/MaterialRenderItem.cpp b/Engine/Source/Developer/MaterialBaking/Private/MaterialRenderItem.cpp
--- a/Engine/Source/Developer/MaterialBaking/Private/MaterialRenderItem.cpp
+++ b/Engine/Source/Developer/MaterialBaking/Private/MaterialRenderItem.cpp
@@ -10,8 +10,17 @@
 
 #define SHOW_WIREFRAME_MESH 0
 
-FMeshMaterialRenderItem::FMeshMaterialRenderItem(const FMaterialData* InMaterialSettings, const FMeshData* InMeshSettings, EMaterialProperty InMaterialProperty)
-	: MeshSettings(InMeshSettings), MaterialSettings(InMaterialSettings), MaterialProperty(InMaterialProperty), MaterialRenderProxy(nullptr), ViewFamily(nullptr)
+FMeshMaterialRenderItem::FMeshMaterialRenderItem(
+	const FMaterialData* InMaterialSettings,
+	const FMeshData* InMeshSettings,
+	EMaterialProperty InMaterialProperty,
+	FDynamicMeshBufferAllocator* InDynamicMeshBufferAllocator)
+	: MeshSettings(InMeshSettings)
+	, MaterialSettings(InMaterialSettings)
+	, MaterialProperty(InMaterialProperty)
+	, MaterialRenderProxy(nullptr)
+	, ViewFamily(nullptr)
+	, DynamicMeshBufferAllocator(InDynamicMeshBufferAllocator)
 {
 	GenerateRenderData();
 	LCI = new FMeshRenderInfo(InMeshSettings->LightMap, nullptr, nullptr, InMeshSettings->LightmapResourceCluster);
@@ -101,7 +110,9 @@ void FMeshMaterialRenderItem::GenerateRenderData()
 
 void FMeshMaterialRenderItem::QueueMaterial(FRHICommandListImmediate& RHICmdList, FMeshPassProcessorRenderState& DrawRenderState, const FSceneView* View)
 {
-	FDynamicMeshBuilder DynamicMeshBuilder(View->GetFeatureLevel(), MAX_STATIC_TEXCOORDS, MeshSettings->LightMapIndex);
+	TRACE_CPUPROFILER_EVENT_SCOPE(FMeshMaterialRenderItem::QueueMaterial)
+
+	FDynamicMeshBuilder DynamicMeshBuilder(View->GetFeatureLevel(), MAX_STATIC_TEXCOORDS, MeshSettings->LightMapIndex, false, DynamicMeshBufferAllocator);
 	DynamicMeshBuilder.AddVertices(Vertices);
 	DynamicMeshBuilder.AddTriangles(Indices);
 
diff --git a/Engine/Source/Developer/MaterialBaking/Private/MaterialRenderItem.h b/Engine/Source/Developer/MaterialBaking/Private/MaterialRenderItem.h
--- a/Engine/Source/Developer/MaterialBaking/Private/MaterialRenderItem.h
+++ b/Engine/Source/Developer/MaterialBaking/Private/MaterialRenderItem.h
@@ -17,7 +17,7 @@ struct FMeshPassProcessorRenderState;
 class FMeshMaterialRenderItem : public FCanvasBaseRenderItem
 {
 public:
-	FMeshMaterialRenderItem(const FMaterialData* InMaterialSettings, const FMeshData* InMeshSettings, EMaterialProperty InMaterialProperty);
+	FMeshMaterialRenderItem(const FMaterialData* InMaterialSettings, const FMeshData* InMeshSettings, EMaterialProperty InMaterialProperty, FDynamicMeshBufferAllocator* InDynamicMeshBufferAllocator = nullptr);
 
 	/** Begin FCanvasBaseRenderItem overrides */
 	virtual bool Render_RenderThread(FRHICommandListImmediate& RHICmdList, FMeshPassProcessorRenderState& DrawRenderState, const FCanvas* Canvas) final;
@@ -48,4 +48,6 @@ public:
 	FLightCacheInterface* LCI;
 	/** View family to use while baking */
 	FSceneViewFamily* ViewFamily;
+private:
+	FDynamicMeshBufferAllocator* DynamicMeshBufferAllocator;
 };
diff --git a/Engine/Source/Developer/MaterialBaking/Public/MaterialBakingHelpers.h b/Engine/Source/Developer/MaterialBaking/Public/MaterialBakingHelpers.h
--- a/Engine/Source/Developer/MaterialBaking/Public/MaterialBakingHelpers.h
+++ b/Engine/Source/Developer/MaterialBaking/Public/MaterialBakingHelpers.h
@@ -10,11 +10,21 @@ class MATERIALBAKING_API FMaterialBakingHelpers
 {
 public:
 	/** 
-	 * Applies a box blur to magenta pixels found in given texture represented by InBMP using non-magenta pixels, this creates a smear across the magenta/filled pixels 
-	 * @param	InOutPixels		The image to apply the smear to
-	 * @param	ImageWidth		The width of the image in pixels
-	 * @param	ImageHeight		The height of the image in pixels
-	 * @param	MaxIterations	The max distance in pixels to smear the edges of the texture. When set to -1 this value is the max of the width and height of the buffer
+	 * Applies a box blur to magenta pixels found in given texture represented by InOutPixels using non-magenta pixels, this creates a smear across the magenta/filled pixels.
+	 * @param  InOutPixels     The image to apply the smear to
+	 * @param  ImageWidth      The width of the image in pixels
+	 * @param  ImageHeight     The height of the image in pixels
+	 * @param  MaxIterations   The max distance in pixels to smear the edges of the texture. When set to -1 this value is the max of the width and height of the buffer
 	 */
 	static void PerformUVBorderSmear(TArray<FColor>& InOutPixels, int32 ImageWidth, int32 ImageHeight, int32 MaxIterations = -1);
+
+	/** 
+	 * Applies a box blur to magenta pixels found in given texture represented by InOutPixels using non-magenta pixels, this creates a smear across the magenta/filled pixels.
+	 * @param  InOutPixels      The image to apply the smear to
+	 * @param  InOutImageWidth  The width of the image in pixels
+	 * @param  InOutImageHeight The height of the image in pixels
+	 *
+	 * @note   This version is allowed to shrink the array to a single pixel if the entire output is monochrome. Params will be adjusted accordingly when this occurs.
+	 */
+	static void PerformUVBorderSmearAndShrink(TArray<FColor>& InOutPixels, int32& InOutImageWidth, int32& InOutImageHeight);
 };
\ No newline at end of file
diff --git a/Engine/Source/Runtime/Engine/Private/DynamicMeshBuilder.cpp b/Engine/Source/Runtime/Engine/Private/DynamicMeshBuilder.cpp
--- a/Engine/Source/Runtime/Engine/Private/DynamicMeshBuilder.cpp
+++ b/Engine/Source/Runtime/Engine/Private/DynamicMeshBuilder.cpp
@@ -229,39 +229,140 @@ void FDynamicMeshIndexBuffer16::UpdateRHI()
 	RHIUnlockIndexBuffer(IndexBufferRHI);
 }
 
+/** FDynamicMeshBufferAllocator's base implementation. It always reallocates new buffers. */
+FDynamicMeshBufferAllocator::~FDynamicMeshBufferAllocator()
+{
+}
+
+int32 FDynamicMeshBufferAllocator::GetIndexBufferSize(uint32 NumElements) const
+{
+	return NumElements * sizeof(DynamicMeshIndexType);
+}
+
+int32 FDynamicMeshBufferAllocator::GetVertexBufferSize(uint32 Stride, uint32 NumElements) const
+{
+	return NumElements * Stride;
+}
+
+FIndexBufferRHIRef FDynamicMeshBufferAllocator::AllocIndexBuffer(uint32 NumElements)
+{
+	uint32 SizeInBytes = GetIndexBufferSize(NumElements);
+
+	FRHIResourceCreateInfo CreateInfo;
+	return RHICreateIndexBuffer(sizeof(DynamicMeshIndexType), SizeInBytes, BUF_Volatile, CreateInfo);
+}
+
+void FDynamicMeshBufferAllocator::ReleaseIndexBuffer(FIndexBufferRHIRef& IndexBufferRHI)
+{
+	IndexBufferRHI = nullptr;
+}
+
+FVertexBufferRHIRef FDynamicMeshBufferAllocator::AllocVertexBuffer(uint32 Stride, uint32 NumElements)
+{
+	uint32 SizeInBytes = GetVertexBufferSize(Stride, NumElements);
+
+	FRHIResourceCreateInfo CreateInfo;
+	return RHICreateVertexBuffer(SizeInBytes, BUF_Volatile | BUF_ShaderResource, CreateInfo);
+}
+
+void FDynamicMeshBufferAllocator::ReleaseVertexBuffer(FVertexBufferRHIRef& VertexBufferRHI)
+{
+	VertexBufferRHI = nullptr;
+}
+
+/** This is our default implementation using GDynamicMeshIndexPool. */
+class FPooledDynamicMeshBufferAllocator : public FDynamicMeshBufferAllocator
+{
+	virtual FIndexBufferRHIRef AllocIndexBuffer(uint32 NumElements) override
+	{
+		uint32 SizeInBytes = NumElements * sizeof(DynamicMeshIndexType);
+		if (SizeInBytes <= FGlobalDynamicMeshIndexPolicy().GetPoolBucketSize(FGlobalDynamicMeshIndexPolicy::NumPoolBuckets - 1))
+		{
+			return GDynamicMeshIndexPool.CreatePooledResource(SizeInBytes);
+		}
+
+		return FDynamicMeshBufferAllocator::AllocIndexBuffer(NumElements);
+	}
+
+	virtual void ReleaseIndexBuffer(FIndexBufferRHIRef& IndexBufferRHI)
+	{
+		if (IsValidRef(IndexBufferRHI))
+		{
+			if (IndexBufferRHI->GetSize() <= FGlobalDynamicMeshIndexPolicy().GetPoolBucketSize(FGlobalDynamicMeshIndexPolicy::NumPoolBuckets - 1))
+			{
+				GDynamicMeshIndexPool.ReleasePooledResource(IndexBufferRHI);
+			}
+
+			IndexBufferRHI = nullptr;
+		}
+	}
+
+	virtual FVertexBufferRHIRef AllocVertexBuffer(uint32 Stride, uint32 NumElements)
+	{
+		uint32 SizeInBytes = NumElements * Stride;
+		if (SizeInBytes <= FGlobalDynamicMeshIndexPolicy().GetPoolBucketSize(FGlobalDynamicMeshIndexPolicy::NumPoolBuckets - 1))
+		{
+			return GDynamicMeshVertexPool.CreatePooledResource(SizeInBytes);
+		}
+
+		return FDynamicMeshBufferAllocator::AllocVertexBuffer(Stride, NumElements);
+	}
+
+	virtual void ReleaseVertexBuffer(FVertexBufferRHIRef& VertexBufferRHI)
+	{
+		if (IsValidRef(VertexBufferRHI))
+		{
+			if (VertexBufferRHI->GetSize() <= FGlobalDynamicMeshVertexPolicy().GetPoolBucketSize(FGlobalDynamicMeshVertexPolicy::NumPoolBuckets - 1))
+			{
+				GDynamicMeshVertexPool.ReleasePooledResource(VertexBufferRHI);
+			}
+
+			VertexBufferRHI = nullptr;
+		}
+	}
+};
+
+static FPooledDynamicMeshBufferAllocator DefaultDynamicMeshBufferAllocator;
+
 /** The index buffer type used for dynamic meshes. */
 class FPooledDynamicMeshIndexBuffer : public FDynamicPrimitiveResource, public FIndexBuffer
 {
 public:
 	TArray<DynamicMeshIndexType> Indices;
 
+	FPooledDynamicMeshIndexBuffer(FDynamicMeshBufferAllocator& InDynamicMeshBufferAllocator)
+		: DynamicMeshBufferAllocator(InDynamicMeshBufferAllocator)
+	{
+	}
+
 	// FRenderResource interface.
 	virtual void InitRHI() override
 	{
-		uint32 SizeInBytes = Indices.Num() * sizeof(DynamicMeshIndexType);
-		if(SizeInBytes <= FGlobalDynamicMeshIndexPolicy().GetPoolBucketSize(FGlobalDynamicMeshIndexPolicy::NumPoolBuckets - 1))
+		TRACE_CPUPROFILER_EVENT_SCOPE(FPooledDynamicMeshIndexBuffer::InitRHI)
+
+		IndexBufferRHI = DynamicMeshBufferAllocator.AllocIndexBuffer(Indices.Num());
+
+		// Write the indices to the index buffer.
+		void* Buffer;
 		{
-			IndexBufferRHI = GDynamicMeshIndexPool.CreatePooledResource(SizeInBytes);
+			TRACE_CPUPROFILER_EVENT_SCOPE(RHILockIndexBuffer)
+			Buffer = RHILockIndexBuffer(IndexBufferRHI,0,Indices.Num() * sizeof(DynamicMeshIndexType),RLM_WriteOnly);
 		}
-		else
 		{
-			FRHIResourceCreateInfo CreateInfo;
-			IndexBufferRHI = RHICreateIndexBuffer(sizeof(DynamicMeshIndexType), SizeInBytes, BUF_Volatile, CreateInfo);
+			TRACE_CPUPROFILER_EVENT_SCOPE(Memcpy)
+			FMemory::Memcpy(Buffer, Indices.GetData(),Indices.Num() * sizeof(DynamicMeshIndexType));
+		}
+		{
+			TRACE_CPUPROFILER_EVENT_SCOPE(RHIUnlockIndexBuffer)
+			RHIUnlockIndexBuffer(IndexBufferRHI);
 		}
-		
-		// Write the indices to the index buffer.
-		void* Buffer = RHILockIndexBuffer(IndexBufferRHI,0,Indices.Num() * sizeof(DynamicMeshIndexType),RLM_WriteOnly);
-		FMemory::Memcpy(Buffer,Indices.GetData(),Indices.Num() * sizeof(DynamicMeshIndexType));
-		RHIUnlockIndexBuffer(IndexBufferRHI);
 	}
 	
 	virtual void ReleaseRHI() override
 	{
-		if(IndexBufferRHI->GetSize() <= FGlobalDynamicMeshIndexPolicy().GetPoolBucketSize(FGlobalDynamicMeshIndexPolicy::NumPoolBuckets - 1))
-		{
-			GDynamicMeshIndexPool.ReleasePooledResource(IndexBufferRHI);
-			IndexBufferRHI = NULL;
-		}
+		TRACE_CPUPROFILER_EVENT_SCOPE(FPooledDynamicMeshIndexBuffer::ReleaseRHI)
+
+		DynamicMeshBufferAllocator.ReleaseIndexBuffer(IndexBufferRHI);
 		FIndexBuffer::ReleaseRHI();
 	}
 
@@ -275,34 +376,10 @@ public:
 		ReleaseResource();
 		delete this;
 	}
-};
 
-static FVertexBufferRHIRef AllocVertexBuffer(uint32 Stride, uint32 NumElements)
-{
-	FVertexBufferRHIRef VertexBufferRHI;
-	uint32 SizeInBytes = NumElements * Stride;
-	if (SizeInBytes <= FGlobalDynamicMeshIndexPolicy().GetPoolBucketSize(FGlobalDynamicMeshIndexPolicy::NumPoolBuckets - 1))
-	{
-		VertexBufferRHI = GDynamicMeshVertexPool.CreatePooledResource(SizeInBytes);
-	}
-	else
-	{
-		FRHIResourceCreateInfo CreateInfo;
-		VertexBufferRHI = RHICreateVertexBuffer(SizeInBytes, BUF_Volatile | BUF_ShaderResource, CreateInfo);
-	}
-
-	return VertexBufferRHI;
-}
-
-static void ReleaseVertexBuffer(FVertexBuffer& VertexBuffer)
-{
-	if (IsValidRef(VertexBuffer.VertexBufferRHI)
-		&& VertexBuffer.VertexBufferRHI->GetSize() <= FGlobalDynamicMeshVertexPolicy().GetPoolBucketSize(FGlobalDynamicMeshVertexPolicy::NumPoolBuckets - 1))
-	{
-		GDynamicMeshVertexPool.ReleasePooledResource(VertexBuffer.VertexBufferRHI);
-		VertexBuffer.VertexBufferRHI = nullptr;
-	}
-}
+private:
+	FDynamicMeshBufferAllocator& DynamicMeshBufferAllocator;
+};
 
 /** The vertex buffer type used for dynamic meshes. */
 class FPooledDynamicMeshVertexBuffer : public FDynamicPrimitiveResource, public FRenderResource
@@ -320,7 +397,11 @@ public:
 
 	TArray<FDynamicMeshVertex> Vertices;
 
-	FPooledDynamicMeshVertexBuffer(uint32 InNumTexCoords, uint32 InLightmapCoordinateIndex, bool InUse16bitTexCoord) : NumTexCoords(InNumTexCoords), LightmapCoordinateIndex(InLightmapCoordinateIndex), Use16bitTexCoord(InUse16bitTexCoord)
+	FPooledDynamicMeshVertexBuffer(uint32 InNumTexCoords, uint32 InLightmapCoordinateIndex, bool InUse16bitTexCoord, FDynamicMeshBufferAllocator& InDynamicMeshBufferAllocator)
+		: NumTexCoords(InNumTexCoords)
+		, LightmapCoordinateIndex(InLightmapCoordinateIndex)
+		, Use16bitTexCoord(InUse16bitTexCoord)
+		, DynamicMeshBufferAllocator(InDynamicMeshBufferAllocator)
 	{
 		check(NumTexCoords > 0 && NumTexCoords <= MAX_STATIC_TEXCOORDS);
 		check(LightmapCoordinateIndex < NumTexCoords);
@@ -329,6 +410,8 @@ public:
 	// FRenderResource interface.
 	virtual void InitRHI() override
 	{
+		TRACE_CPUPROFILER_EVENT_SCOPE(FPooledDynamicMeshVertexBuffer::InitRHI)
+
 		uint32 TextureStride = sizeof(FVector2D);
 		EPixelFormat TextureFormat = PF_G32R32F;
 
@@ -338,10 +421,10 @@ public:
 			TextureFormat = PF_G16R16F;
 		}
 
-		PositionBuffer.VertexBufferRHI	= AllocVertexBuffer(sizeof(FVector), Vertices.Num());
-		TangentBuffer.VertexBufferRHI	= AllocVertexBuffer(sizeof(FPackedNormal), 2 * Vertices.Num());
-		TexCoordBuffer.VertexBufferRHI	= AllocVertexBuffer(TextureStride, NumTexCoords * Vertices.Num());
-		ColorBuffer.VertexBufferRHI		= AllocVertexBuffer(sizeof(FColor), Vertices.Num());
+		PositionBuffer.VertexBufferRHI = DynamicMeshBufferAllocator.AllocVertexBuffer(sizeof(FVector), Vertices.Num());
+		TangentBuffer.VertexBufferRHI  = DynamicMeshBufferAllocator.AllocVertexBuffer(sizeof(FPackedNormal), 2 * Vertices.Num());
+		TexCoordBuffer.VertexBufferRHI = DynamicMeshBufferAllocator.AllocVertexBuffer(TextureStride, NumTexCoords * Vertices.Num());
+		ColorBuffer.VertexBufferRHI    = DynamicMeshBufferAllocator.AllocVertexBuffer(sizeof(FColor), Vertices.Num());
 
 		if (RHISupportsManualVertexFetch(GMaxRHIShaderPlatform))
 		{
@@ -361,22 +444,28 @@ public:
 		FPackedNormal* TangentBufferData	= static_cast<FPackedNormal*>(RHILockVertexBuffer(TangentBuffer.VertexBufferRHI, 0, 2 * sizeof(FPackedNormal) * Vertices.Num(), RLM_WriteOnly));	
 		FColor* ColorBufferData				= static_cast<FColor*>(RHILockVertexBuffer(ColorBuffer.VertexBufferRHI, 0, sizeof(FColor) * Vertices.Num(), RLM_WriteOnly));
 
-		for (int32 i = 0; i < Vertices.Num(); i++)
 		{
-			PositionBufferData[i] = Vertices[i].Position;
-			TangentBufferData[2 * i + 0] = Vertices[i].TangentX;
-			TangentBufferData[2 * i + 1] = Vertices[i].TangentZ;
-			ColorBufferData[i] = Vertices[i].Color;
+			// This code will generate a lot of page faults when the memory has never been written to
+			// so we'll know pooled buffers are not behaving optimally if this shows up more than we would expect.
+			TRACE_CPUPROFILER_EVENT_SCOPE(CopyDataToVertexBuffers)
 
-			for (uint32 j = 0; j < NumTexCoords; j++)
+			for (int32 i = 0; i < Vertices.Num(); i++)
 			{
-				if (Use16bitTexCoord)
-				{
-					TexCoordBufferData16[NumTexCoords * i + j] = FVector2DHalf(Vertices[i].TextureCoordinate[j]);
-				}
-				else
+				PositionBufferData[i] = Vertices[i].Position;
+				TangentBufferData[2 * i + 0] = Vertices[i].TangentX;
+				TangentBufferData[2 * i + 1] = Vertices[i].TangentZ;
+				ColorBufferData[i] = Vertices[i].Color;
+
+				for (uint32 j = 0; j < NumTexCoords; j++)
 				{
-					TexCoordBufferData32[NumTexCoords * i + j] = Vertices[i].TextureCoordinate[j];
+					if (Use16bitTexCoord)
+					{
+						TexCoordBufferData16[NumTexCoords * i + j] = FVector2DHalf(Vertices[i].TextureCoordinate[j]);
+					}
+					else
+					{
+						TexCoordBufferData32[NumTexCoords * i + j] = Vertices[i].TextureCoordinate[j];
+					}
 				}
 			}
 		}
@@ -407,10 +496,10 @@ public:
 
 	virtual void ReleaseRHI() override
 	{
-		ReleaseVertexBuffer(PositionBuffer);
-		ReleaseVertexBuffer(TangentBuffer);
-		ReleaseVertexBuffer(TexCoordBuffer);
-		ReleaseVertexBuffer(ColorBuffer);
+		DynamicMeshBufferAllocator.ReleaseVertexBuffer(PositionBuffer.VertexBufferRHI);
+		DynamicMeshBufferAllocator.ReleaseVertexBuffer(TangentBuffer.VertexBufferRHI);
+		DynamicMeshBufferAllocator.ReleaseVertexBuffer(TexCoordBuffer.VertexBufferRHI);
+		DynamicMeshBufferAllocator.ReleaseVertexBuffer(ColorBuffer.VertexBufferRHI);
 	}
 
 	// FDynamicPrimitiveResource interface.
@@ -443,6 +532,7 @@ private:
 	const uint32 NumTexCoords;
 	const uint32 LightmapCoordinateIndex;
 	const bool Use16bitTexCoord;
+	FDynamicMeshBufferAllocator& DynamicMeshBufferAllocator;
 };
 
 /** The vertex factory type used for dynamic meshes. */
@@ -566,12 +656,18 @@ public:
 	}
 };
 
-FDynamicMeshBuilder::FDynamicMeshBuilder(ERHIFeatureLevel::Type InFeatureLevel, uint32 InNumTexCoords, uint32 InLightmapCoordinateIndex, bool InUse16bitTexCoord)
+FDynamicMeshBuilder::FDynamicMeshBuilder(ERHIFeatureLevel::Type InFeatureLevel, uint32 InNumTexCoords, uint32 InLightmapCoordinateIndex, bool InUse16bitTexCoord, FDynamicMeshBufferAllocator* InDynamicMeshBufferAllocator)
 	: FeatureLevel(InFeatureLevel)
 {
 	OneFrameResources = nullptr;
-	VertexBuffer = new FPooledDynamicMeshVertexBuffer(InNumTexCoords, InLightmapCoordinateIndex, InUse16bitTexCoord);
-	IndexBuffer = new FPooledDynamicMeshIndexBuffer;
+
+	if (InDynamicMeshBufferAllocator == nullptr)
+	{
+		InDynamicMeshBufferAllocator = &DefaultDynamicMeshBufferAllocator;
+	}
+
+	VertexBuffer = new FPooledDynamicMeshVertexBuffer(InNumTexCoords, InLightmapCoordinateIndex, InUse16bitTexCoord, *InDynamicMeshBufferAllocator);
+	IndexBuffer = new FPooledDynamicMeshIndexBuffer(*InDynamicMeshBufferAllocator);
 }
 
 FDynamicMeshBuilder::~FDynamicMeshBuilder()
@@ -731,6 +827,8 @@ void FDynamicMeshBuilder::GetMesh(
 	FMeshElementCollector& Collector,
 	const FHitProxyId HitProxyId )
 {
+	TRACE_CPUPROFILER_EVENT_SCOPE(FDynamicMeshBuilder::GetMesh)
+
 	// Only draw non-empty meshes.
 	if((VertexBuffer && VertexBuffer->Vertices.Num() > 0) || (DrawOffset != nullptr))
 	{
@@ -825,6 +923,8 @@ void FDynamicMeshBuilder::GetMesh(
 
 void FDynamicMeshBuilder::GetMeshElement(const FMatrix& LocalToWorld, const FMaterialRenderProxy* MaterialRenderProxy, uint8 DepthPriorityGroup, bool bDisableBackfaceCulling, bool bReceivesDecals, int32 ViewIndex, FMeshBuilderOneFrameResources& OneFrameResource, FMeshBatch& Mesh)
 {
+	TRACE_CPUPROFILER_EVENT_SCOPE(FDynamicMeshBuilder::GetMeshElement)
+
 	// Only draw non-empty meshes.
 	if (VertexBuffer->Vertices.Num() > 0)
 	{
@@ -908,6 +1008,8 @@ void FDynamicMeshBuilder::GetMeshElement(const FMatrix& LocalToWorld, const FMat
 
 void FDynamicMeshBuilder::Draw(FPrimitiveDrawInterface* PDI,const FMatrix& LocalToWorld,const FMaterialRenderProxy* MaterialRenderProxy,uint8 DepthPriorityGroup,bool bDisableBackfaceCulling, bool bReceivesDecals, const FHitProxyId HitProxyId )
 {
+	TRACE_CPUPROFILER_EVENT_SCOPE(FDynamicMeshBuilder::Draw)
+
 	// Only draw non-empty meshes.
 	if(VertexBuffer->Vertices.Num() > 0 && IndexBuffer->Indices.Num() > 0)
 	{
diff --git a/Engine/Source/Runtime/Engine/Public/DynamicMeshBuilder.h b/Engine/Source/Runtime/Engine/Public/DynamicMeshBuilder.h
--- a/Engine/Source/Runtime/Engine/Public/DynamicMeshBuilder.h
+++ b/Engine/Source/Runtime/Engine/Public/DynamicMeshBuilder.h
@@ -135,6 +135,23 @@ struct FDynamicMeshBuilderSettings
 	bool bUseWireframeSelectionColoring = false;
 };
 
+/**
+ * This class provides the vertex/index allocator interface used by FDynamicMeshBuilder which is
+ * already implemented internally with caching in mind but can be customized if needed.
+ */
+class ENGINE_API FDynamicMeshBufferAllocator
+{
+public:
+	int32 GetIndexBufferSize(uint32 NumElements) const;
+	int32 GetVertexBufferSize(uint32 Stride, uint32 NumElements) const;
+	virtual ~FDynamicMeshBufferAllocator();
+
+	virtual FIndexBufferRHIRef AllocIndexBuffer(uint32 NumElements);
+	virtual void ReleaseIndexBuffer(FIndexBufferRHIRef& IndexBufferRHI);
+	virtual FVertexBufferRHIRef AllocVertexBuffer(uint32 Stride, uint32 NumElements);
+	virtual void ReleaseVertexBuffer(FVertexBufferRHIRef& VertexBufferRHI);
+};
+
 /**
  * A utility used to construct dynamically generated meshes, and render them to a FPrimitiveDrawInterface.
  * Note: This is meant to be easy to use, not fast.  It moves the data around more than necessary, and requires dynamically allocating RHI
@@ -145,7 +162,7 @@ class FDynamicMeshBuilder
 public:
 
 	/** Initialization constructor. */
-	ENGINE_API FDynamicMeshBuilder(ERHIFeatureLevel::Type InFeatureLevel, uint32 InNumTexCoords = 1, uint32 InLightmapCoordinateIndex = 0, bool InUse16bitTexCoord = false);
+	ENGINE_API FDynamicMeshBuilder(ERHIFeatureLevel::Type InFeatureLevel, uint32 InNumTexCoords = 1, uint32 InLightmapCoordinateIndex = 0, bool InUse16bitTexCoord = false, FDynamicMeshBufferAllocator* InDynamicMeshBufferAllocator = nullptr);
 
 	/** Destructor. */
 	ENGINE_API ~FDynamicMeshBuilder();
@@ -206,6 +223,7 @@ private:
 	class FPooledDynamicMeshIndexBuffer* IndexBuffer = nullptr;
 	class FPooledDynamicMeshVertexBuffer* VertexBuffer = nullptr;
 	ERHIFeatureLevel::Type FeatureLevel;
+	FDynamicMeshBufferAllocator* DynamicMeshBufferAllocator;
 };
 
 /** Index Buffer */
