Fixed crash when generating an incorrect node type connected to "GetPropery" of material options.

diff --git a/Engine/Source/Developer/MaterialBaking/Private/MaterialOptionsCustomization.cpp b/Engine/Source/Developer/MaterialBaking/Private/MaterialOptionsCustomization.cpp
--- a/Engine/Source/Developer/MaterialBaking/Private/MaterialOptionsCustomization.cpp
+++ b/Engine/Source/Developer/MaterialBaking/Private/MaterialOptionsCustomization.cpp
@@ -37,8 +37,10 @@ void FPropertyEntryCustomization::CustomizeHeader(TSharedRef<IPropertyHandle> Pr
 	];
 
 	const TArray<TWeakObjectPtr<UObject>>& SelectedObjects = CustomizationUtils.GetPropertyUtilities()->GetSelectedObjects();
+
 	// Try and find material options instance in currently edited objects
-	CurrentOptions = Cast<UMaterialOptions>((SelectedObjects.FindByPredicate([](TWeakObjectPtr<UObject> Object) { return Cast<UMaterialOptions>(Object.Get()); }))->Get());
+	const TWeakObjectPtr<UObject>* MaterialOptionsPtr = SelectedObjects.FindByPredicate([](TWeakObjectPtr<UObject> Object) { return Cast<UMaterialOptions>(Object.Get()); });
+	CurrentOptions = MaterialOptionsPtr ? Cast<UMaterialOptions>((MaterialOptionsPtr)->Get()) : nullptr;
 	
 	const int32 Index = PropertyHandle->GetIndexInArray();
 
