Merge CL 11587939 - Fix invalid material baking when "Reuse Mesh Lightmap UVs is used with a mirrored (negative scale) mesh.

diff --git a/Engine/Source/Developer/MaterialBaking/Private/MaterialRenderItem.cpp b/Engine/Source/Developer/MaterialBaking/Private/MaterialRenderItem.cpp
--- a/Engine/Source/Developer/MaterialBaking/Private/MaterialRenderItem.cpp
+++ b/Engine/Source/Developer/MaterialBaking/Private/MaterialRenderItem.cpp
@@ -227,11 +227,16 @@ void FMeshMaterialRenderItem::PopulateWithMeshData()
 		{
 			if (MeshSettings->MaterialIndices.Contains(PolygonGroupID.GetValue()))
 			{
-				for (int32 Corner = 0; Corner < 3; Corner++)
+				const int32 NUM_VERTICES = 3;
+				for (int32 Corner = 0; Corner < NUM_VERTICES; Corner++)
 				{
-					const int32 SrcVertIndex = FaceIndex * 3 + Corner;
+					// Swap vertices order if mesh is mirrored
+					const int32 CornerIdx = !MeshSettings->bMirrored ? Corner : NUM_VERTICES - Corner - 1;
+
+					const int32 SrcVertIndex = FaceIndex * NUM_VERTICES + CornerIdx;
 					const FVertexInstanceID SrcVertexInstanceID = RawMesh->GetTriangleVertexInstance(TriangleID, Corner);
 					const FVertexID SrcVertexID = RawMesh->GetVertexInstanceVertex(SrcVertexInstanceID);
+
 					// add vertex
 					FDynamicMeshVertex* Vert = new(Vertices)FDynamicMeshVertex();
 					if (!bUseNewUVs)
@@ -275,12 +280,6 @@ void FMeshMaterialRenderItem::PopulateWithMeshData()
 				Indices.Add(VertIndex - 3);
 				Indices.Add(VertIndex - 1);
 				Indices.Add(VertIndex - 2);
-
-				// Swap vertices order if mesh is mirrored
-				if (MeshSettings->bMirrored)
-				{
-					Swap(Vertices[Vertices.Num() - 3].Position, Vertices[Vertices.Num() - 1].Position);
-				}
 			}
 			FaceIndex++;
 		}
