Edigrate 14806872 from UE4-Main to avoid crashes when baking materials during HLOD generation.

diff --git a/Engine/Source/Developer/MaterialBaking/Private/MaterialBakingModule.cpp b/Engine/Source/Developer/MaterialBaking/Private/MaterialBakingModule.cpp
--- a/Engine/Source/Developer/MaterialBaking/Private/MaterialBakingModule.cpp
+++ b/Engine/Source/Developer/MaterialBaking/Private/MaterialBakingModule.cpp
@@ -286,10 +286,9 @@ void FMaterialBakingModule::ShutdownModule()
 
 void FMaterialBakingModule::BakeMaterials(const TArray<FMaterialData*>& MaterialSettings, const TArray<FMeshData*>& MeshSettings, TArray<FBakeOutput>& Output)
 {
-	TArray<FMaterialDataEx> MaterialDataExs;
-	TArray<FMaterialDataEx*> MaterialSettingsEx;
-
 	// Translate old material data to extended types
+	TArray<FMaterialDataEx> MaterialDataExs;
+	MaterialDataExs.Reserve(MaterialSettings.Num());
 	for (const FMaterialData* MaterialData : MaterialSettings)
 	{
 		FMaterialDataEx& MaterialDataEx = MaterialDataExs.AddDefaulted_GetRef();
@@ -300,7 +299,13 @@ void FMaterialBakingModule::BakeMaterials(const TArray<FMaterialData*>& Material
 		{
 			MaterialDataEx.PropertySizes.Add(PropertySizePair.Key, PropertySizePair.Value);
 		}
+	}
 
+	// Build an array of pointers to the extended type
+	TArray<FMaterialDataEx*> MaterialSettingsEx;
+	MaterialSettingsEx.Reserve(MaterialDataExs.Num());
+	for (FMaterialDataEx& MaterialDataEx : MaterialDataExs)
+	{
 		MaterialSettingsEx.Add(&MaterialDataEx);
 	}
 
@@ -308,6 +313,7 @@ void FMaterialBakingModule::BakeMaterials(const TArray<FMaterialData*>& Material
 	BakeMaterials(MaterialSettingsEx, MeshSettings, BakeOutputExs);
 
 	// Translate extended bake output to old types
+	Output.Reserve(BakeOutputExs.Num());
 	for (FBakeOutputEx& BakeOutputEx : BakeOutputExs)
 	{
 		FBakeOutput& BakeOutput = Output.AddDefaulted_GetRef();
