MaterialBaking - Fixed FMaterialBakingModule::CleanupMaterialProxies() deleting the material proxy on the wrong thread

diff --git a/Engine/Source/Developer/MaterialBaking/Private/MaterialBakingModule.cpp b/Engine/Source/Developer/MaterialBaking/Private/MaterialBakingModule.cpp
--- a/Engine/Source/Developer/MaterialBaking/Private/MaterialBakingModule.cpp
+++ b/Engine/Source/Developer/MaterialBaking/Private/MaterialBakingModule.cpp
@@ -23,6 +23,7 @@
 #include "Misc/ScopedSlowTask.h"
 #include "MeshDescription.h"
 #include "TextureCompiler.h"
+#include "RenderCaptureInterface.h"
 #if WITH_EDITOR
 #include "Misc/FileHelper.h"
 #endif
@@ -48,6 +49,14 @@ static TAutoConsoleVariable<int32> CVarSaveIntermediateTextures(
 	TEXT("1: Turned On"),
 	ECVF_Default);
 
+static TAutoConsoleVariable<int32> CVarMaterialBakingRDOCCapture(
+	TEXT("MaterialBaking.RenderDocCapture"),
+	0,
+	TEXT("Determines whether or not to trigger a RenderDoc capture.\n")
+	TEXT("0: Turned Off\n")
+	TEXT("1: Turned On"),
+	ECVF_Default);
+
 namespace FMaterialBakingModuleImpl
 {
 	// Custom dynamic mesh allocator specifically tailored for Material Baking.
@@ -333,6 +342,8 @@ void FMaterialBakingModule::BakeMaterials(const TArray<FMaterialData*>& Material
 
 void FMaterialBakingModule::BakeMaterials(const TArray<FMaterialDataEx*>& MaterialSettings, const TArray<FMeshData*>& MeshSettings, TArray<FBakeOutputEx>& Output)
 {
+	RenderCaptureInterface::FScopedCapture RenderCapture(CVarMaterialBakingRDOCCapture.GetValueOnAnyThread() == 1, TEXT("MaterialBaking"));
+
 	TRACE_CPUPROFILER_EVENT_SCOPE(FMaterialBakingModule::BakeMaterials)
 
 	checkf(MaterialSettings.Num() == MeshSettings.Num(), TEXT("Number of material settings does not match that of MeshSettings"));
@@ -734,11 +745,20 @@ bool FMaterialBakingModule::SetupMaterialBakeSettings(TArray<TWeakObjectPtr<UObj
 	return false;
 }
 
+static void DeleteCachedMaterialProxy(FExportMaterialProxy* Proxy)
+{
+	ENQUEUE_RENDER_COMMAND(DeleteCachedMaterialProxy)(
+		[Proxy](FRHICommandListImmediate& RHICmdList)
+		{
+			delete Proxy;
+		});
+}
+
 void FMaterialBakingModule::CleanupMaterialProxies()
 {
 	for (auto Iterator : MaterialProxyPool)
 	{
-		delete Iterator.Value.Value;
+		DeleteCachedMaterialProxy(Iterator.Value.Value);
 	}
 	MaterialProxyPool.Reset();
 }
@@ -947,14 +967,7 @@ void FMaterialBakingModule::OnObjectModified(UObject* Object)
 				// We have a match, remove the entry from our pool
 				if (bMustDelete)
 				{
-					FExportMaterialProxy* Proxy = It.Value().Value;
-
-					ENQUEUE_RENDER_COMMAND(DeleteCachedMaterialProxy)(
-						[Proxy](FRHICommandListImmediate& RHICmdList)
-						{
-							delete Proxy;
-						});
-
+					DeleteCachedMaterialProxy(It.Value().Value);
 					It.RemoveCurrent();
 				}
 			}
