MaterialBaking - Removed parallelism in PerformUVSmear as it is actually degrading performance (really small workloads) & we've now parallelized more parts of the material baking process

diff --git a/Engine/Source/Developer/MaterialBaking/Private/MaterialBakingHelpers.cpp b/Engine/Source/Developer/MaterialBaking/Private/MaterialBakingHelpers.cpp
--- a/Engine/Source/Developer/MaterialBaking/Private/MaterialBakingHelpers.cpp
+++ b/Engine/Source/Developer/MaterialBaking/Private/MaterialBakingHelpers.cpp
@@ -2,7 +2,6 @@
 
 #include "MaterialBakingHelpers.h"
 #include "Math/Color.h"
-#include "Async/ParallelFor.h"
 
 namespace FMaterialBakingHelpersImpl
 {
@@ -214,8 +213,6 @@ namespace FMaterialBakingHelpersImpl
 		TArray<uint32> RowRemainingPixels;
 		RowRemainingPixels.SetNumZeroed(PaddedImageHeight);
 
-		const int32 MaxThreads = FPlatformProcess::SupportsMultithreading() ? FPlatformMisc::NumberOfCoresIncludingHyperthreads() : 1;
-
 		//
 		// Iteratively smear until all rows are filled.
 		//
@@ -223,40 +220,30 @@ namespace FMaterialBakingHelpersImpl
 		int32 LoopCount = 0;
 		while (CurrentRowsLeft->Num() && (MaxIterations <= 0 || LoopCount <= MaxIterations))
 		{
-			const int32 NumThreads = FMath::Min(CurrentRowsLeft->Num(), MaxThreads);
-			const int32 LinesPerThread = FMath::CeilToInt((float)CurrentRowsLeft->Num() / (float)NumThreads);
-
-			// split up rows that still have magenta pixels amongst threads
-			ParallelFor(NumThreads, [ImageWidth, PaddedImageWidth, PaddedImageHeight, Current, Scratch, CurrentRowsLeft, LinesPerThread, &RowRemainingPixels, MagentaMask](int32 ThreadIndex)
+			for (int32 Index = 0; Index < CurrentRowsLeft->Num(); Index++)
 			{
-				const int32 StartY = ThreadIndex*LinesPerThread;
-				const int32 EndY = FMath::Min((ThreadIndex + 1) * LinesPerThread, CurrentRowsLeft->Num());
+				const int32 Y = (*CurrentRowsLeft)[Index];
+				RowRemainingPixels[Index] = 0;
 
-				for (int32 Index = StartY; Index < EndY; Index++)
+				int32 PixelIndex = (Y - 1) * ImageWidth;
+				for (int32 X = 1; X <= ImageWidth; X++)
 				{
-					const int32 Y = (*CurrentRowsLeft)[Index];
-					RowRemainingPixels[Index] = 0;
-
-					int32 PixelIndex = (Y - 1) * ImageWidth;
-					for (int32 X = 1; X <= ImageWidth; X++)
+					FColor& Color = Current[PixelIndex++];
+					if (Color.DWColor() == MagentaMask)
 					{
-						FColor& Color = Current[PixelIndex++];
-						if (Color.DWColor() == MagentaMask)
+						const FColor SampledColor = BoxBlurSample(Scratch, X, Y, PaddedImageWidth, PaddedImageHeight);
+						// If it's a valid pixel
+						if (SampledColor.DWColor() != MagentaMask)
 						{
-							const FColor SampledColor = BoxBlurSample(Scratch, X, Y, PaddedImageWidth, PaddedImageHeight);
-							// If it's a valid pixel
-							if (SampledColor.DWColor() != MagentaMask)
-							{
-								Color = SampledColor;
-							}
-							else
-							{
-								RowRemainingPixels[Index]++;
-							}
+							Color = SampledColor;
+						}
+						else
+						{
+							RowRemainingPixels[Index]++;
 						}
 					}
 				}
-			});
+			}
 
 			// Mark all completed rows and copy data between buffers
 			const int32 RowMemorySize = ImageWidth * sizeof(FColor);
