MaterialBaking module - cleanup cached material proxies upon module shutdown & GC collect passes

diff --git a/Engine/Source/Developer/MaterialBaking/Private/MaterialBakingModule.cpp b/Engine/Source/Developer/MaterialBaking/Private/MaterialBakingModule.cpp
--- a/Engine/Source/Developer/MaterialBaking/Private/MaterialBakingModule.cpp
+++ b/Engine/Source/Developer/MaterialBaking/Private/MaterialBakingModule.cpp
@@ -279,6 +279,9 @@ void FMaterialBakingModule::StartupModule()
 	
 	// Register callback for modified objects
 	FCoreUObjectDelegates::OnObjectModified.AddRaw(this, &FMaterialBakingModule::OnObjectModified);
+
+	// Register callback on garbage collection
+	FCoreUObjectDelegates::GetPreGarbageCollectDelegate().AddRaw(this, &FMaterialBakingModule::OnPreGarbageCollect);
 }
 
 void FMaterialBakingModule::ShutdownModule()
@@ -290,7 +293,11 @@ void FMaterialBakingModule::ShutdownModule()
 	{
 		PropertyEditorModule->UnregisterCustomPropertyTypeLayout(TEXT("PropertyEntry"));
 	}
+
 	FCoreUObjectDelegates::OnObjectModified.RemoveAll(this);
+	FCoreUObjectDelegates::GetPreGarbageCollectDelegate().RemoveAll(this);
+
+	CleanupMaterialProxies();
 }
 
 void FMaterialBakingModule::BakeMaterials(const TArray<FMaterialData*>& MaterialSettings, const TArray<FMeshData*>& MeshSettings, TArray<FBakeOutput>& Output)
@@ -975,4 +982,9 @@ void FMaterialBakingModule::OnObjectModified(UObject* Object)
 	}
 }
 
+void FMaterialBakingModule::OnPreGarbageCollect()
+{
+	CleanupMaterialProxies();
+}
+
 #undef LOCTEXT_NAMESPACE //"MaterialBakingModule"
diff --git a/Engine/Source/Developer/MaterialBaking/Private/MaterialBakingModule.h b/Engine/Source/Developer/MaterialBaking/Private/MaterialBakingModule.h
--- a/Engine/Source/Developer/MaterialBaking/Private/MaterialBakingModule.h
+++ b/Engine/Source/Developer/MaterialBaking/Private/MaterialBakingModule.h
@@ -52,6 +52,10 @@ protected:
 
 	/** Callback for modified objects which should be removed from MaterialProxyPool in that case */
 	void OnObjectModified(UObject* Object);
+
+	/** Callback used to clear material proxy cache on garbage collection */
+	void OnPreGarbageCollect();
+
 private:
 	/** Pool of available render targets, cached for re-using on consecutive property rendering */
 	TArray<UTextureRenderTarget2D*> RenderTargetPool;
