Fixed crash when baking a material for export from a mac.
From my observation Metal doesn't like copying a texture that is SRGB to one that is not SRGB.

diff --git a/Engine/Source/Developer/MaterialBaking/Private/MaterialBakingModule.cpp b/Engine/Source/Developer/MaterialBaking/Private/MaterialBakingModule.cpp
--- a/Engine/Source/Developer/MaterialBaking/Private/MaterialBakingModule.cpp
+++ b/Engine/Source/Developer/MaterialBaking/Private/MaterialBakingModule.cpp
@@ -154,14 +154,14 @@ namespace FMaterialBakingModuleImpl
 	class FStagingBufferPool
 	{
 	public:
-		FTexture2DRHIRef CreateStagingBuffer_RenderThread(FRHICommandListImmediate& RHICmdList, int32 Width, int32 Height, EPixelFormat Format)
+		FTexture2DRHIRef CreateStagingBuffer_RenderThread(FRHICommandListImmediate& RHICmdList, int32 Width, int32 Height, EPixelFormat Format, bool bIsSRGB)
 		{
 			TRACE_CPUPROFILER_EVENT_SCOPE(CreateStagingBuffer_RenderThread)
 
 			auto StagingBufferPredicate = 
-				[Width, Height, Format](const FTexture2DRHIRef& Texture2DRHIRef)
+				[Width, Height, Format, bIsSRGB](const FTexture2DRHIRef& Texture2DRHIRef)
 				{
-					return Texture2DRHIRef->GetSizeX() == Width && Texture2DRHIRef->GetSizeY() == Height && Texture2DRHIRef->GetFormat() == Format;
+					return Texture2DRHIRef->GetSizeX() == Width && Texture2DRHIRef->GetSizeY() == Height && Texture2DRHIRef->GetFormat() == Format && bool(Texture2DRHIRef->GetFlags() & TexCreate_SRGB) == bIsSRGB;
 				};
 
 			// Process any staging buffers available for unmapping
@@ -191,7 +191,13 @@ namespace FMaterialBakingModuleImpl
 
 			TRACE_CPUPROFILER_EVENT_SCOPE(RHICreateTexture2D)
 			FRHIResourceCreateInfo CreateInfo(TEXT("FStagingBufferPool_StagingBuffer"));
-			return RHICreateTexture2D(Width, Height, Format, 1, 1, TexCreate_CPUReadback, CreateInfo);
+			ETextureCreateFlags TextureCreateFlags = TexCreate_CPUReadback;
+			if (bIsSRGB)
+			{
+				TextureCreateFlags |= TexCreate_SRGB;
+			}
+			
+			return RHICreateTexture2D(Width, Height, Format, 1, 1, TextureCreateFlags, CreateInfo);
 		}
 
 		void ReleaseStagingBufferForUnmap_AnyThread(FTexture2DRHIRef& Texture2DRHIRef)
@@ -551,12 +557,12 @@ void FMaterialBakingModule::BakeMaterials(const TArray<FMaterialDataEx*>& Materi
 				}
 
 				// Lookup gamma and format settings for property, if not found use default values
-				bool InForceLinearGamma = PerPropertyGamma.Contains(Property);
+				const bool bForceLinearGamma = PerPropertyGamma.Contains(Property);
 				EPixelFormat PixelFormat = PerPropertyFormat.Contains(Property) ? PerPropertyFormat[Property] : PF_B8G8R8A8;
 
 				// It is safe to reuse the same render target for each draw pass since they all execute sequentially on the GPU and are copied to staging buffers before
 				// being reused.
-				UTextureRenderTarget2D* RenderTarget = CreateRenderTarget(InForceLinearGamma, PixelFormat, CurrentOutput.PropertySizes[Property]);
+				UTextureRenderTarget2D* RenderTarget = CreateRenderTarget(bForceLinearGamma, PixelFormat, CurrentOutput.PropertySizes[Property]);
 				if (RenderTarget != nullptr)
 				{
 					// Perform everything left of the operation directly on the render thread since we need to modify some RenderItem's properties
@@ -591,7 +597,7 @@ void FMaterialBakingModule::BakeMaterials(const TArray<FMaterialDataEx*>& Materi
 							Canvas.Flush_RenderThread(RHICmdList);
 							SortElement.RenderBatchArray.Empty();
 
-							FTexture2DRHIRef StagingBufferRef = StagingBufferPool.CreateStagingBuffer_RenderThread(RHICmdList, RenderTargetResource->GetSizeX(), RenderTargetResource->GetSizeY(), RenderTarget->GetFormat());
+							FTexture2DRHIRef StagingBufferRef = StagingBufferPool.CreateStagingBuffer_RenderThread(RHICmdList, RenderTargetResource->GetSizeX(), RenderTargetResource->GetSizeY(), RenderTarget->GetFormat(), RenderTarget->IsSRGB());
 							FGPUFenceRHIRef GPUFence = RHICreateGPUFence(TEXT("MaterialBackingFence"));
 
 							FResolveRect Rect(0, 0, RenderTargetResource->GetSizeX(), RenderTargetResource->GetSizeY());
