MaterialBaking - Cap smearing iterations to avoid extreme processing time on high resolution baked textures

diff --git a/Engine/Source/Developer/MaterialBaking/Private/MaterialBakingHelpers.cpp b/Engine/Source/Developer/MaterialBaking/Private/MaterialBakingHelpers.cpp
--- a/Engine/Source/Developer/MaterialBaking/Private/MaterialBakingHelpers.cpp
+++ b/Engine/Source/Developer/MaterialBaking/Private/MaterialBakingHelpers.cpp
@@ -307,13 +307,17 @@ namespace FMaterialBakingHelpersImpl
 		}
 
 		// If we finished before replacing all pixels, replace the remaining pixels with black.
-		if (CurrentRowsLeft->Num() > 0)
+		for (int32 Index = 0; Index < CurrentRowsLeft->Num(); Index++)
 		{
-			for (int32 i = 0; i < InOutPixels.Num(); ++i)
+			const int32 Y = (*CurrentRowsLeft)[Index];
+
+			int32 PixelIndex = (Y - 1) * ImageWidth;
+			for (int32 X = 1; X <= ImageWidth; X++)
 			{
-				if (InOutPixels[i].DWColor() == MagentaMask)
+				FColor& Color = Current[PixelIndex++];
+				if (Color.DWColor() == MagentaMask)
 				{
-					InOutPixels[i] = FColor::Black;
+					Color = FColor::Black;
 				}
 			}
 		}
@@ -327,5 +331,6 @@ void FMaterialBakingHelpers::PerformUVBorderSmear(TArray<FColor>& InOutPixels, i
 
 void FMaterialBakingHelpers::PerformUVBorderSmearAndShrink(TArray<FColor>& InOutPixels, int32& InOutImageWidth, int32& InOutImageHeight)
 {
-	FMaterialBakingHelpersImpl::PerformUVBorderSmear(InOutPixels, InOutImageWidth, InOutImageHeight, -1, true);
+	const int32 DefaultMaxIteration = 8; // This amount of smearing should be enough to cover any bleeding issue.
+	FMaterialBakingHelpersImpl::PerformUVBorderSmear(InOutPixels, InOutImageWidth, InOutImageHeight, DefaultMaxIteration, true);
 }
\ No newline at end of file
