PR #7916: Lucasfilm: MaterialBaking: Fixes and improvements for material baking (Contributed by [REDACTED_PERSON])
- Fixes for baking decal materials
- Option to return full HDR emissive color value
- Option to make all values linear by default, except for colors

diff --git a/Engine/Source/Developer/MaterialBaking/Private/ExportMaterialProxy.h b/Engine/Source/Developer/MaterialBaking/Private/ExportMaterialProxy.h
--- a/Engine/Source/Developer/MaterialBaking/Private/ExportMaterialProxy.h
+++ b/Engine/Source/Developer/MaterialBaking/Private/ExportMaterialProxy.h
@@ -463,10 +463,8 @@ public:
 	}
 	virtual EMaterialDomain GetMaterialDomain() const override
 	{
-		if (Material)
-		{
-			return Material->MaterialDomain;
-		}
+		// Because the baking module applies the material to a plane (or mesh),
+		// it needs to be a surface material.
 		return MD_Surface;
 	}
 	virtual bool IsTwoSided() const  override
@@ -495,7 +493,10 @@ public:
 	}
 	virtual bool IsDeferredDecal() const override
 	{
-		return Material && Material->MaterialDomain == MD_DeferredDecal;
+		// Decals are tricky. Since they mix with the underlying material
+		// and can't be applied to meshes, they can't really be baked 1:1.
+		// Instead we'll just bake them as surface materials.
+		return false;
 	}
 	virtual bool IsVolumetricPrimitive() const override
 	{
diff --git a/Engine/Source/Developer/MaterialBaking/Private/MaterialBakingModule.cpp b/Engine/Source/Developer/MaterialBaking/Private/MaterialBakingModule.cpp
--- a/Engine/Source/Developer/MaterialBaking/Private/MaterialBakingModule.cpp
+++ b/Engine/Source/Developer/MaterialBaking/Private/MaterialBakingModule.cpp
@@ -258,11 +258,10 @@ namespace FMaterialBakingModuleImpl
 
 void FMaterialBakingModule::StartupModule()
 {
+	bEmissiveHDR = false;
+
 	// Set which properties should enforce gamma correction
-	PerPropertyGamma.Add(MP_Normal);
-	PerPropertyGamma.Add(MP_Opacity);
-	PerPropertyGamma.Add(MP_OpacityMask);
-	PerPropertyGamma.Add(TEXT("ClearCoatBottomNormal"));
+	SetLinearBake(false);
 
 	// Set which pixel format should be used for the possible baked out material properties
 	PerPropertyFormat.Add(MP_EmissiveColor, PF_FloatRGBA);
@@ -352,6 +351,11 @@ void FMaterialBakingModule::BakeMaterials(const TArray<FMaterialData*>& Material
 		{
 			BakeOutput.PropertyData.Add(PropertyDataPair.Key.Type, MoveTemp(PropertyDataPair.Value));
 		}
+
+		for (TPair<FMaterialPropertyEx, TArray<FFloat16Color>>& PropertyDataPair : BakeOutputEx.HDRPropertyData)
+		{
+			BakeOutput.HDRPropertyData.Add(PropertyDataPair.Key.Type, MoveTemp(PropertyDataPair.Value));
+		}
 	}
 }
 
@@ -540,6 +544,10 @@ void FMaterialBakingModule::BakeMaterials(const TArray<FMaterialDataEx*>& Materi
 			{
 				const FMaterialPropertyEx& Property = MaterialPropertiesToBakeOut[PropertyIndex];
 				CurrentOutput.PropertyData.Add(Property);
+				if (bEmissiveHDR && Property == MP_EmissiveColor)
+				{
+					CurrentOutput.HDRPropertyData.Add(Property);
+				}
 			}
 
 			for (int32 PropertyIndex = 0; PropertyIndex < NumPropertiesToRender; ++PropertyIndex)
@@ -557,12 +565,13 @@ void FMaterialBakingModule::BakeMaterials(const TArray<FMaterialDataEx*>& Materi
 				}
 
 				// Lookup gamma and format settings for property, if not found use default values
-				const bool bForceLinearGamma = PerPropertyGamma.Contains(Property);
-				EPixelFormat PixelFormat = PerPropertyFormat.Contains(Property) ? PerPropertyFormat[Property] : PF_B8G8R8A8;
+				const EPropertyColorSpace* OverrideColorSpace = PerPropertyColorSpace.Find(Property);
+				const EPropertyColorSpace ColorSpace = OverrideColorSpace ? *OverrideColorSpace : DefaultColorSpace;
+				const EPixelFormat PixelFormat = PerPropertyFormat.Contains(Property) ? PerPropertyFormat[Property] : PF_B8G8R8A8;
 
 				// It is safe to reuse the same render target for each draw pass since they all execute sequentially on the GPU and are copied to staging buffers before
 				// being reused.
-				UTextureRenderTarget2D* RenderTarget = CreateRenderTarget(bForceLinearGamma, PixelFormat, CurrentOutput.PropertySizes[Property]);
+				UTextureRenderTarget2D* RenderTarget = CreateRenderTarget((ColorSpace == EPropertyColorSpace::Linear), PixelFormat, CurrentOutput.PropertySizes[Property]);
 				if (RenderTarget != nullptr)
 				{
 					// Perform everything left of the operation directly on the render thread since we need to modify some RenderItem's properties
@@ -607,7 +616,7 @@ void FMaterialBakingModule::BakeMaterials(const TArray<FMaterialDataEx*>& Materi
 							// Prepare a lambda for final processing that will be executed asynchronously
 							NumTasks++;
 							auto FinalProcessing_AnyThread =
-								[&NumTasks, bSaveIntermediateTextures, CurrentMaterialSettings, &StagingBufferPool, &Output, Property, MaterialIndex](FTexture2DRHIRef& StagingBuffer, void * Data, int32 DataWidth, int32 DataHeight)
+								[&NumTasks, bSaveIntermediateTextures, CurrentMaterialSettings, &StagingBufferPool, &Output, Property, MaterialIndex, bEmissiveHDR = bEmissiveHDR](FTexture2DRHIRef& StagingBuffer, void * Data, int32 DataWidth, int32 DataHeight)
 								{
 									TRACE_CPUPROFILER_EVENT_SCOPE(FinalProcessing)
 
@@ -621,6 +630,13 @@ void FMaterialBakingModule::BakeMaterials(const TArray<FMaterialDataEx*>& Materi
 									{
 										// Only one thread will write to CurrentOutput.EmissiveScale since there can be only one emissive channel property per FBakeOutputEx
 										FMaterialBakingModule::ProcessEmissiveOutput((const FFloat16Color*)Data, DataWidth, OutputSize, OutputColor, CurrentOutput.EmissiveScale);
+
+										if (bEmissiveHDR)
+										{
+											TArray<FFloat16Color>& OutputHDRColor = CurrentOutput.HDRPropertyData[Property];
+											OutputHDRColor.SetNum(OutputSize.X * OutputSize.Y);
+											ConvertRawR16G16B16A16FDataToFFloat16Color(OutputSize.X, OutputSize.Y, (uint8*)Data, DataWidth * sizeof(FFloat16Color), OutputHDRColor.GetData());
+										}
 									}
 									else
 									{
@@ -774,6 +790,32 @@ bool FMaterialBakingModule::SetupMaterialBakeSettings(TArray<TWeakObjectPtr<UObj
 	return false;
 }
 
+void FMaterialBakingModule::SetEmissiveHDR(bool bHDR)
+{
+	bEmissiveHDR = bHDR;
+}
+
+void FMaterialBakingModule::SetLinearBake(bool bCorrectLinear)
+{
+	// PerPropertyGamma ultimately sets whether the render target is linear
+	PerPropertyColorSpace.Reset();
+	if (bCorrectLinear)
+	{
+		DefaultColorSpace = EPropertyColorSpace::Linear;
+		PerPropertyColorSpace.Add(MP_BaseColor, EPropertyColorSpace::sRGB);
+		PerPropertyColorSpace.Add(MP_EmissiveColor, EPropertyColorSpace::sRGB);
+		PerPropertyColorSpace.Add(MP_SubsurfaceColor, EPropertyColorSpace::sRGB);
+	}
+	else
+	{
+		DefaultColorSpace = EPropertyColorSpace::sRGB;
+		PerPropertyColorSpace.Add(MP_Normal, EPropertyColorSpace::Linear);
+		PerPropertyColorSpace.Add(MP_Opacity, EPropertyColorSpace::Linear);
+		PerPropertyColorSpace.Add(MP_OpacityMask, EPropertyColorSpace::Linear);
+		PerPropertyColorSpace.Add(TEXT("ClearCoatBottomNormal"), EPropertyColorSpace::Linear);
+	}
+}
+
 static void DeleteCachedMaterialProxy(FExportMaterialProxy* Proxy)
 {
 	ENQUEUE_RENDER_COMMAND(DeleteCachedMaterialProxy)(
diff --git a/Engine/Source/Developer/MaterialBaking/Private/MaterialBakingModule.h b/Engine/Source/Developer/MaterialBaking/Private/MaterialBakingModule.h
--- a/Engine/Source/Developer/MaterialBaking/Private/MaterialBakingModule.h
+++ b/Engine/Source/Developer/MaterialBaking/Private/MaterialBakingModule.h
@@ -37,6 +37,12 @@ public:
 	/** Promps a slate window to allow the user to populate specific material baking settings used while baking out materials */
 	virtual bool SetupMaterialBakeSettings(TArray<TWeakObjectPtr<UObject>>& OptionObjects, int32 NumLODs) override;
 
+	/** Outputs true HDR version of emissive color */
+	virtual void SetEmissiveHDR(bool bHDR) override;
+
+	/** Bakes all material properties to linear textures, except for colors */
+	virtual void SetLinearBake(bool bCorrectLinear) override;
+
 protected:
 	/* Creates and adds or reuses a RenderTarget from the pool */
 	UTextureRenderTarget2D* CreateRenderTarget(bool bInForceLinearGamma, EPixelFormat InPixelFormat, const FIntPoint& InTargetSize);
@@ -57,6 +63,12 @@ protected:
 	void OnPreGarbageCollect();
 
 private:
+	enum EPropertyColorSpace
+	{
+		Linear,
+		sRGB,
+	};
+
 	/** Pool of available render targets, cached for re-using on consecutive property rendering */
 	TArray<UTextureRenderTarget2D*> RenderTargetPool;
 
@@ -70,5 +82,9 @@ private:
 	TMap<FMaterialPropertyEx, EPixelFormat> PerPropertyFormat;
 
 	/** Whether or not to enforce gamma correction while baking out specific material properties */
-	TSet<FMaterialPropertyEx> PerPropertyGamma;
+	TMap<FMaterialPropertyEx, EPropertyColorSpace> PerPropertyColorSpace;
+
+	EPropertyColorSpace DefaultColorSpace;
+
+	bool bEmissiveHDR;
 };
diff --git a/Engine/Source/Developer/MaterialBaking/Public/IMaterialBakingModule.h b/Engine/Source/Developer/MaterialBaking/Public/IMaterialBakingModule.h
--- a/Engine/Source/Developer/MaterialBaking/Public/IMaterialBakingModule.h
+++ b/Engine/Source/Developer/MaterialBaking/Public/IMaterialBakingModule.h
@@ -31,4 +31,10 @@ public:
 
 	/** Promps a slate window to allow the user to populate specific material baking settings used while baking out materials */
 	virtual  bool SetupMaterialBakeSettings(TArray<TWeakObjectPtr<UObject>>& OptionObjects, int32 NumLODs) = 0;
+
+	/** Outputs true HDR version of emissive color */
+	virtual void SetEmissiveHDR(bool bHDR) = 0;
+
+	/** Bakes all material properties to linear textures, except for colors */
+	virtual void SetLinearBake(bool bCorrectLinear) = 0;
 };
diff --git a/Engine/Source/Developer/MaterialBaking/Public/MaterialBakingStructures.h b/Engine/Source/Developer/MaterialBaking/Public/MaterialBakingStructures.h
--- a/Engine/Source/Developer/MaterialBaking/Public/MaterialBakingStructures.h
+++ b/Engine/Source/Developer/MaterialBaking/Public/MaterialBakingStructures.h
@@ -95,6 +95,9 @@ struct FBakeOutput
 	/** Contains the resulting texture size for baking out a material's property */
 	TMap<EMaterialProperty, FIntPoint> PropertySizes;
 
+	/** Contains the resulting HDR texture data for baking out a material's property, may be empty */
+	TMap<EMaterialProperty, TArray<FFloat16Color>> HDRPropertyData;
+
 	/** Scale used to allow having wide ranges of emissive values in the source materials, the final proxy material will use this value to scale the emissive texture's pixel values */
 	float EmissiveScale;
 };
@@ -112,6 +115,9 @@ struct FBakeOutputEx
 	/** Contains the resulting texture size for baking out a extened material's property */
 	TMap<FMaterialPropertyEx, FIntPoint> PropertySizes;
 
+	/** Contains the resulting HDR texture data for baking out a material's property, may be empty */
+	TMap<FMaterialPropertyEx, TArray<FFloat16Color>> HDRPropertyData;
+
 	/** Scale used to allow having wide ranges of emissive values in the source materials, the final proxy material will use this value to scale the emissive texture's pixel values */
 	float EmissiveScale;
 };
diff --git a/Engine/Source/Runtime/RHI/Public/RHISurfaceDataConversion.h b/Engine/Source/Runtime/RHI/Public/RHISurfaceDataConversion.h
--- a/Engine/Source/Runtime/RHI/Public/RHISurfaceDataConversion.h
+++ b/Engine/Source/Runtime/RHI/Public/RHISurfaceDataConversion.h
@@ -91,6 +91,29 @@ static inline void ConvertRawB8G8R8A8DataToFColor(uint32 Width, uint32 Height, u
 	}
 }
 
+static inline void ConvertRawR16G16B16A16FDataToFFloat16Color(uint32 Width, uint32 Height, uint8* In, uint32 SrcPitch, FFloat16Color* Out)
+{
+	const uint32 DstPitch = Width * sizeof(FFloat16Color);
+
+	// If source & dest pitch matches, perform a single memcpy.
+	if (DstPitch == SrcPitch)
+	{
+		FPlatformMemory::Memcpy(Out, In, Width * Height * sizeof(FFloat16Color));
+	}
+	else
+	{
+		check(SrcPitch > DstPitch);
+
+		// Need to copy row wise since the Pitch does not match the Width.
+		for (uint32 Y = 0; Y < Height; Y++)
+		{
+			FFloat16Color* SrcPtr = (FFloat16Color*)(In + Y * SrcPitch);
+			FFloat16Color* DestPtr = Out + Y * Width;
+			FMemory::Memcpy(DestPtr, SrcPtr, DstPitch);
+		}
+	}
+}
+
 static inline void ConvertRawR10G10B10A2DataToFColor(uint32 Width, uint32 Height, uint8 *In, uint32 SrcPitch, FColor* Out)
 {
 	for (uint32 Y = 0; Y < Height; Y++)
