FMaterialBakingHelpers::PerformUVBorderSmearAndShrink() should shrink images regardless of the number of smearing iterations

diff --git a/Engine/Source/Developer/MaterialBaking/Private/MaterialBakingHelpers.cpp b/Engine/Source/Developer/MaterialBaking/Private/MaterialBakingHelpers.cpp
--- a/Engine/Source/Developer/MaterialBaking/Private/MaterialBakingHelpers.cpp
+++ b/Engine/Source/Developer/MaterialBaking/Private/MaterialBakingHelpers.cpp
@@ -67,12 +67,59 @@ namespace FMaterialBakingHelpersImpl
 		return false;
 	}
 
-	void PerformUVBorderSmear(TArray<FColor>& InOutPixels, int32& ImageWidth, int32& ImageHeight, int32 MaxIterations, bool bCanShrink)
+	static void PerformShrinking(TArray<FColor>& InOutPixels, int32& ImageWidth, int32& ImageHeight)
+	{
+		const uint32 MagentaMask = FColor(255, 0, 255).DWColor();
+
+		FColor* Current = InOutPixels.GetData();
+
+		// We can hopefully skip the entire smearing process if there is a single
+		// non-magenta color in the entire array since smearing that completely
+		// would lead to a monochome output.
+		uint32 SingleColor = MagentaMask;
+		bool   bHasSingleColor = true;
+		for (int32 Index = 0, Num = InOutPixels.Num(); Index < Num; ++Index)
+		{
+			const uint32 Color = Current[Index].DWColor();
+			if (Color != MagentaMask)
+			{
+				if (SingleColor == MagentaMask)
+				{
+					// This is the first time we stumble on a color other than magenta, keep it.
+					SingleColor = Color;
+				}
+				// Compare with the known non-magenta color
+				else if (SingleColor != Color)
+				{
+					bHasSingleColor = false;
+					break;
+				}
+			}
+		}
+
+		// Shrink to a single texel
+		if (bHasSingleColor)
+		{
+			InOutPixels.SetNum(1);
+			ImageWidth = ImageHeight = 1;
+			Current = InOutPixels.GetData();
+
+			InOutPixels[0] = FColor(SingleColor);
+		}
+	}
+
+	static void PerformUVBorderSmear(TArray<FColor>& InOutPixels, int32& ImageWidth, int32& ImageHeight, int32 MaxIterations)
 	{
 		TRACE_CPUPROFILER_EVENT_SCOPE(FMaterialBakingHelpers::PerformUVBorderSmear)
 
 		check(InOutPixels.Num() == ImageWidth*ImageHeight);
 
+		// No smearing needed for a single pixel
+		if (ImageWidth * ImageHeight == 1)
+		{
+			return;
+		}
+
 		const uint32 MagentaMask = FColor(255, 0, 255).DWColor();
 		const int32 PaddedImageWidth = ImageWidth + 2;
 		const int32 PaddedImageHeight = ImageHeight + 2;
@@ -87,54 +134,6 @@ namespace FMaterialBakingHelpersImpl
 		// This avoids needing to bounds check in inner loops.
 		//
 
-		FColor* Current = InOutPixels.GetData();
-
-		// We can hopefully skip the entire smearing process if there is a single
-		// non-magenta color in the entire array since smearing that completely
-		// would lead to a monochome output.
-		if (MaxIterations == -1)
-		{
-			uint32 SingleColor     = MagentaMask;
-			bool   bHasSingleColor = true;
-			for (int32 Index = 0, Num = InOutPixels.Num(); Index < Num; ++Index)
-			{
-				const uint32 Color = Current[Index].DWColor();
-				if (Color != MagentaMask)
-				{
-					if (SingleColor == MagentaMask)
-					{
-						// This is the first time we stumble on a color other than magenta, keep it.
-						SingleColor = Color;
-					}
-					// Compare with the known non-magenta color
-					else if (SingleColor != Color)
-					{
-						bHasSingleColor = false;
-						break;
-					}
-				}
-			}
-
-			if (bHasSingleColor)
-			{
-				// Shrink to single texel when allowed
-				if (bCanShrink)
-				{
-					InOutPixels.SetNum(1);
-					ImageWidth = ImageHeight = 1;
-					Current = InOutPixels.GetData();
-				}
-
-				FColor Color(SingleColor);
-				for (int32 Index = 0, Num = InOutPixels.Num(); Index < Num; ++Index)
-				{
-					Current[Index] = Color;
-				}
-
-				return;
-			}
-		}
-
 		TArray<FColor> ScratchBuffer;
 		ScratchBuffer.SetNumUninitialized(PaddedImageWidth*PaddedImageHeight);
 		FColor* Scratch = ScratchBuffer.GetData();
@@ -146,7 +145,9 @@ namespace FMaterialBakingHelpersImpl
 			Scratch[(PaddedImageHeight - 1)*PaddedImageWidth + X] = FColor(MagentaMask);
 		}
 
-		// Set leftg and right border pixels to magenta and copy image data into rows
+		FColor* Current = InOutPixels.GetData();
+
+		// Set left and right border pixels to magenta and copy image data into rows
 		for (int32 Y = 1; Y <= ImageHeight; ++Y)
 		{
 			const int32 YOffset = Y * PaddedImageWidth;
@@ -326,11 +327,13 @@ namespace FMaterialBakingHelpersImpl
 
 void FMaterialBakingHelpers::PerformUVBorderSmear(TArray<FColor>& InOutPixels, int32 ImageWidth, int32 ImageHeight, int32 MaxIterations)
 {
-	FMaterialBakingHelpersImpl::PerformUVBorderSmear(InOutPixels, ImageWidth, ImageHeight, MaxIterations, false);
+	FMaterialBakingHelpersImpl::PerformUVBorderSmear(InOutPixels, ImageWidth, ImageHeight, MaxIterations);
 }
 
 void FMaterialBakingHelpers::PerformUVBorderSmearAndShrink(TArray<FColor>& InOutPixels, int32& InOutImageWidth, int32& InOutImageHeight)
 {
+	FMaterialBakingHelpersImpl::PerformShrinking(InOutPixels, InOutImageWidth, InOutImageHeight);
+
 	const int32 DefaultMaxIteration = 8; // This amount of smearing should be enough to cover any bleeding issue.
-	FMaterialBakingHelpersImpl::PerformUVBorderSmear(InOutPixels, InOutImageWidth, InOutImageHeight, DefaultMaxIteration, true);
+	FMaterialBakingHelpersImpl::PerformUVBorderSmear(InOutPixels, InOutImageWidth, InOutImageHeight, DefaultMaxIteration);
 }
\ No newline at end of file
