Material baking changes in CL 16278086 modified the default baking behavior

diff --git a/Engine/Source/Developer/MaterialBaking/Private/MaterialBakingModule.cpp b/Engine/Source/Developer/MaterialBaking/Private/MaterialBakingModule.cpp
--- a/Engine/Source/Developer/MaterialBaking/Private/MaterialBakingModule.cpp
+++ b/Engine/Source/Developer/MaterialBaking/Private/MaterialBakingModule.cpp
@@ -565,13 +565,12 @@ void FMaterialBakingModule::BakeMaterials(const TArray<FMaterialDataEx*>& Materi
 				}
 
 				// Lookup gamma and format settings for property, if not found use default values
-				const EPropertyColorSpace* OverrideColorSpace = PerPropertyColorSpace.Find(Property);
-				const EPropertyColorSpace ColorSpace = OverrideColorSpace ? *OverrideColorSpace : DefaultColorSpace;
+				const bool bLinearGamma = PerPropertyLinearGamma.Contains(Property);
 				const EPixelFormat PixelFormat = PerPropertyFormat.Contains(Property) ? PerPropertyFormat[Property] : PF_B8G8R8A8;
 
 				// It is safe to reuse the same render target for each draw pass since they all execute sequentially on the GPU and are copied to staging buffers before
 				// being reused.
-				UTextureRenderTarget2D* RenderTarget = CreateRenderTarget((ColorSpace == EPropertyColorSpace::sRGB), PixelFormat, CurrentOutput.PropertySizes[Property]);
+				UTextureRenderTarget2D* RenderTarget = CreateRenderTarget(bLinearGamma, PixelFormat, CurrentOutput.PropertySizes[Property]);
 				if (RenderTarget != nullptr)
 				{
 					// Perform everything left of the operation directly on the render thread since we need to modify some RenderItem's properties
@@ -798,21 +797,19 @@ void FMaterialBakingModule::SetEmissiveHDR(bool bHDR)
 void FMaterialBakingModule::SetLinearBake(bool bCorrectLinear)
 {
 	// PerPropertyGamma ultimately sets whether the render target is linear
-	PerPropertyColorSpace.Reset();
+	PerPropertyLinearGamma.Reset();
 	if (bCorrectLinear)
 	{
-		DefaultColorSpace = EPropertyColorSpace::Linear;
-		PerPropertyColorSpace.Add(MP_BaseColor, EPropertyColorSpace::sRGB);
-		PerPropertyColorSpace.Add(MP_EmissiveColor, EPropertyColorSpace::sRGB);
-		PerPropertyColorSpace.Add(MP_SubsurfaceColor, EPropertyColorSpace::sRGB);
+		PerPropertyLinearGamma.Add(MP_BaseColor);
+		PerPropertyLinearGamma.Add(MP_EmissiveColor);
+		PerPropertyLinearGamma.Add(MP_SubsurfaceColor);
 	}
 	else
 	{
-		DefaultColorSpace = EPropertyColorSpace::sRGB;
-		PerPropertyColorSpace.Add(MP_Normal, EPropertyColorSpace::Linear);
-		PerPropertyColorSpace.Add(MP_Opacity, EPropertyColorSpace::Linear);
-		PerPropertyColorSpace.Add(MP_OpacityMask, EPropertyColorSpace::Linear);
-		PerPropertyColorSpace.Add(TEXT("ClearCoatBottomNormal"), EPropertyColorSpace::Linear);
+		PerPropertyLinearGamma.Add(MP_Normal);
+		PerPropertyLinearGamma.Add(MP_Opacity);
+		PerPropertyLinearGamma.Add(MP_OpacityMask);
+		PerPropertyLinearGamma.Add(TEXT("ClearCoatBottomNormal"));
 	}
 }
 
diff --git a/Engine/Source/Developer/MaterialBaking/Private/MaterialBakingModule.h b/Engine/Source/Developer/MaterialBaking/Private/MaterialBakingModule.h
--- a/Engine/Source/Developer/MaterialBaking/Private/MaterialBakingModule.h
+++ b/Engine/Source/Developer/MaterialBaking/Private/MaterialBakingModule.h
@@ -82,9 +82,7 @@ private:
 	TMap<FMaterialPropertyEx, EPixelFormat> PerPropertyFormat;
 
 	/** Whether or not to enforce gamma correction while baking out specific material properties */
-	TMap<FMaterialPropertyEx, EPropertyColorSpace> PerPropertyColorSpace;
-
-	EPropertyColorSpace DefaultColorSpace;
+	TSet<FMaterialPropertyEx> PerPropertyLinearGamma;
 
 	bool bEmissiveHDR;
 };
