PR #8493: Fix bad fresnel artifacts when material baking with mesh data (Contributed by [REDACTED_PERSON])

diff --git a/Engine/Source/Developer/MaterialBaking/Private/ExportMaterialProxy.h b/Engine/Source/Developer/MaterialBaking/Private/ExportMaterialProxy.h
--- a/Engine/Source/Developer/MaterialBaking/Private/ExportMaterialProxy.h
+++ b/Engine/Source/Developer/MaterialBaking/Private/ExportMaterialProxy.h
@@ -119,7 +119,8 @@ struct FExportMaterialCompiler : public FProxyMaterialCompiler
 
 	virtual int32 CameraVector() override
 	{
-		return Compiler->Constant3(0.0f, 0.0f, 1.0f);
+		// By returning vertex normal instead of a constant vector (like up), we ensure materials (with fresnel for example) are more correctly baked using custom mesh data.
+		return Compiler->VertexNormal();
 	}
 
 	virtual int32 ReflectionAboutCustomWorldNormal(int32 CustomWorldNormal, int32 bNormalizeCustomWorldNormal) override
