PR #8530: Fix incorrect results when baking materials that use reflection vectors (Contributed by [REDACTED_PERSON])

#rnx

diff --git a/Engine/Source/Developer/MaterialBaking/Private/ExportMaterialProxy.h b/Engine/Source/Developer/MaterialBaking/Private/ExportMaterialProxy.h
--- a/Engine/Source/Developer/MaterialBaking/Private/ExportMaterialProxy.h
+++ b/Engine/Source/Developer/MaterialBaking/Private/ExportMaterialProxy.h
@@ -125,7 +125,22 @@ struct FExportMaterialCompiler : public FProxyMaterialCompiler
 
 	virtual int32 ReflectionAboutCustomWorldNormal(int32 CustomWorldNormal, int32 bNormalizeCustomWorldNormal) override
 	{
-		return Compiler->ReflectionAboutCustomWorldNormal(CustomWorldNormal, bNormalizeCustomWorldNormal);
+		if (CustomWorldNormal == INDEX_NONE)
+		{
+			return INDEX_NONE;
+		}
+
+		int32 N = CustomWorldNormal;
+		int32 C = CameraVector();
+
+		if (bNormalizeCustomWorldNormal)
+		{
+			// N = N / sqrt(dot(N, N))
+			N = Compiler->Div(N, Compiler->SquareRoot(Compiler->Dot(N, N)));
+		}
+
+		// return 2 * dot(N, C) * N - C
+		return Compiler->Sub(Compiler->Mul(Compiler->Constant(2.0f), Compiler->Mul(Compiler->Dot(N, C), N)), C);
 	}
 
 	virtual int32 VertexColor() override
@@ -155,7 +170,8 @@ struct FExportMaterialCompiler : public FProxyMaterialCompiler
 
 	virtual int32 ReflectionVector() override
 	{
-		return Compiler->ReflectionVector();
+		// Because camera vector is identical to normal vector we can work out that reflection vector will also be the same
+		return Compiler->VertexNormal();
 	}
 
 	virtual int32 AtmosphericFogColor(int32 WorldPosition) override
