PR #8492: Add FMaterialPropertyEx::IsCustomOutput (Contributed by [REDACTED_PERSON])

#rnx

diff --git a/Engine/Source/Developer/MaterialBaking/Private/MaterialPropertyEx.cpp b/Engine/Source/Developer/MaterialBaking/Private/MaterialPropertyEx.cpp
new file mode 100644
--- /dev/null
+++ b/Engine/Source/Developer/MaterialBaking/Private/MaterialPropertyEx.cpp
@@ -0,0 +1,18 @@
+// Copyright Epic Games, Inc. All Rights Reserved.
+
+#include "MaterialPropertyEx.h"
+#include "UObject/NameTypes.h"
+
+FString FMaterialPropertyEx::ToString() const
+{
+	if (!IsCustomOutput())
+	{
+		const UEnum* Enum = StaticEnum<EMaterialProperty>();
+		FName Name = Enum->GetNameByValue(Type);
+		FString TrimmedName = Name.ToString();
+		TrimmedName.RemoveFromStart(TEXT("MP_"), ESearchCase::CaseSensitive);
+		return TrimmedName;
+	}
+
+	return CustomOutput.ToString();
+}
diff --git a/Engine/Source/Developer/MaterialBaking/Public/MaterialPropertyEx.h b/Engine/Source/Developer/MaterialBaking/Public/MaterialPropertyEx.h
--- a/Engine/Source/Developer/MaterialBaking/Public/MaterialPropertyEx.h
+++ b/Engine/Source/Developer/MaterialBaking/Public/MaterialPropertyEx.h
@@ -3,7 +3,6 @@
 #pragma once
 
 #include "SceneTypes.h"
-#include "UObject/NameTypes.h"
 
 /** Structure extending EMaterialProperty to allow detailed information about custom output */
 struct FMaterialPropertyEx
@@ -23,9 +22,14 @@ struct FMaterialPropertyEx
 		, CustomOutput(CustomOutput)
 	{}
 
+	FORCEINLINE bool IsCustomOutput() const
+	{
+		return Type == MP_CustomOutput;
+	}
+
 	FORCEINLINE bool operator ==(const FMaterialPropertyEx& Other) const
 	{
-		return Type == Other.Type && (Type != MP_CustomOutput || CustomOutput == Other.CustomOutput);
+		return Type == Other.Type && (!IsCustomOutput() || CustomOutput == Other.CustomOutput);
 	}
 
 	FORCEINLINE bool operator !=(const FMaterialPropertyEx& Other) const
@@ -35,28 +39,14 @@ struct FMaterialPropertyEx
 
 	friend FORCEINLINE uint32 GetTypeHash(const FMaterialPropertyEx& Other)
 	{
-		return Other.Type != MP_CustomOutput ? GetTypeHash(Other.Type) : GetTypeHash(Other.CustomOutput);
+		return !Other.IsCustomOutput() ? GetTypeHash(Other.Type) : GetTypeHash(Other.CustomOutput);
 	}
 
-	FString ToString() const
-	{
-		if (Type != MP_CustomOutput)
-		{
-			const UEnum* Enum = StaticEnum<EMaterialProperty>();
-			FName Name = Enum->GetNameByValue(Type);
-			FString TrimmedName = Name.ToString();
-			TrimmedName.RemoveFromStart(TEXT("MP_"), ESearchCase::CaseSensitive);
-			return TrimmedName;
-		}
-		else
-		{
-			return CustomOutput.ToString();
-		}
-	}
-	
+	MATERIALBAKING_API FString ToString() const;
+
 	/** The material property */
 	EMaterialProperty Type;
-	
+
 	/** The name of a specific custom output. Only used if property is MP_CustomOutput */
 	FName CustomOutput;
 };
