PR #8494: Fix incorrect compiler used for material baking custom outputs (Contributed by [REDACTED_PERSON])

#rnx

diff --git a/Engine/Source/Developer/MaterialBaking/Private/ExportMaterialProxy.h b/Engine/Source/Developer/MaterialBaking/Private/ExportMaterialProxy.h
--- a/Engine/Source/Developer/MaterialBaking/Private/ExportMaterialProxy.h
+++ b/Engine/Source/Developer/MaterialBaking/Private/ExportMaterialProxy.h
@@ -422,7 +422,7 @@ public:
 				return MaterialInterface->CompileProperty(&ProxyCompiler, MP_ShadingModel);
 			case MP_CustomOutput:
 				 // NOTE: Currently we can assume input index is always 0, which it is for all custom outputs that are registered as material attributes
-				return CompileInputForCustomOutput(Compiler, 0, ForceCast_Exact_Replicate);
+				return CompileInputForCustomOutput(&ProxyCompiler, 0, ForceCast_Exact_Replicate);
 			default:
 				return Compiler->Constant(1.0f);
 			}
