PR #8499: Add support for transforming normals in material baking from world-space (Contributed by [REDACTED_PERSON])

#rnx

diff --git a/Engine/Source/Developer/MaterialBaking/Private/ExportMaterialProxy.h b/Engine/Source/Developer/MaterialBaking/Private/ExportMaterialProxy.h
--- a/Engine/Source/Developer/MaterialBaking/Private/ExportMaterialProxy.h
+++ b/Engine/Source/Developer/MaterialBaking/Private/ExportMaterialProxy.h
@@ -245,13 +245,14 @@ struct FExportMaterialCompiler : public FProxyMaterialCompiler
 class FExportMaterialProxy : public FMaterial, public FMaterialRenderProxy
 {
 public:
-	FExportMaterialProxy(UMaterialInterface* InMaterialInterface, EMaterialProperty InPropertyToCompile, const FString& InCustomOutputToCompile = TEXT(""), bool bInSynchronousCompilation = true)
+	FExportMaterialProxy(UMaterialInterface* InMaterialInterface, EMaterialProperty InPropertyToCompile, const FString& InCustomOutputToCompile = TEXT(""), bool bInSynchronousCompilation = true, bool bTangentSpaceNormal = false)
 		: FMaterial()
 		, FMaterialRenderProxy(GetPathNameSafe(InMaterialInterface->GetMaterial()))
 		, MaterialInterface(InMaterialInterface)
 		, PropertyToCompile(InPropertyToCompile)
 		, CustomOutputToCompile(InCustomOutputToCompile)
 		, bSynchronousCompilation(bInSynchronousCompilation)
+		, bTangentSpaceNormal(bTangentSpaceNormal)
 	{
 		SetQualityLevelProperties(GMaxRHIFeatureLevel);
 		Material = InMaterialInterface->GetMaterial();
@@ -415,7 +416,7 @@ public:
 				{
 					return CompileNormalEncoding(
 						Compiler,
-						MaterialInterface->CompileProperty(&ProxyCompiler, PropertyToCompile, ForceCast_Exact_Replicate));
+						CompileNormalTransform(&ProxyCompiler, MaterialInterface->CompileProperty(&ProxyCompiler, PropertyToCompile, ForceCast_Exact_Replicate)));
 				}
 				break;
 			case MP_ShadingModel:
@@ -573,7 +574,7 @@ private:
 
 		if (CustomOutputToCompile == TEXT("ClearCoatBottomNormal"))
 		{
-			Result = CompileNormalEncoding(Compiler, Result);
+			Result = CompileNormalEncoding(Compiler, CompileNormalTransform(Compiler, Result));
 		}
 
 		if (ForceCastFlags & MFCF_ForceCast)
@@ -598,6 +599,12 @@ private:
 		return nullptr;
 	}
 
+	int32 CompileNormalTransform(FMaterialCompiler* Compiler, int32 NormalInput) const
+	{
+		return bTangentSpaceNormal && !Material->bTangentSpaceNormal
+			? Compiler->TransformVector(MCB_World, MCB_Tangent, NormalInput) : NormalInput;
+	}
+
 	static int32 CompileNormalEncoding(FMaterialCompiler* Compiler, int32 NormalInput)
 	{
 		return Compiler->Add(
@@ -616,4 +623,8 @@ private:
 	FString CustomOutputToCompile;
 	FGuid Id;
 	bool bSynchronousCompilation;
+
+public:
+	/** Whether to transform normals from world-space to tangent-space (does nothing if material already uses tangent-space normals) */
+	bool bTangentSpaceNormal;
 };
\ No newline at end of file
diff --git a/Engine/Source/Developer/MaterialBaking/Private/MaterialBakingModule.cpp b/Engine/Source/Developer/MaterialBaking/Private/MaterialBakingModule.cpp
--- a/Engine/Source/Developer/MaterialBaking/Private/MaterialBakingModule.cpp
+++ b/Engine/Source/Developer/MaterialBaking/Private/MaterialBakingModule.cpp
@@ -317,6 +317,7 @@ void FMaterialBakingModule::BakeMaterials(const TArray<FMaterialData*>& Material
 		FMaterialDataEx& MaterialDataEx = MaterialDataExs.AddDefaulted_GetRef();
 		MaterialDataEx.Material = MaterialData->Material;
 		MaterialDataEx.bPerformBorderSmear = MaterialData->bPerformBorderSmear;
+		MaterialDataEx.bTangentSpaceNormal = MaterialData->bTangentSpaceNormal;
 
 		for (const TPair<EMaterialProperty, FIntPoint>& PropertySizePair : MaterialData->PropertySizes)
 		{
@@ -487,7 +488,7 @@ void FMaterialBakingModule::BakeMaterials(const TArray<FMaterialDataEx*>& Materi
 			for (TMap<FMaterialPropertyEx, FIntPoint>::TConstIterator PropertySizeIterator = CurrentMaterialSettings->PropertySizes.CreateConstIterator(); PropertySizeIterator; ++PropertySizeIterator)
 			{
 				// They will be stored in the pool and compiled asynchronously
-				CreateMaterialProxy(CurrentMaterialSettings->Material, PropertySizeIterator.Key());
+				CreateMaterialProxy(CurrentMaterialSettings, PropertySizeIterator.Key());
 			}
 		}
 	}
@@ -556,7 +557,7 @@ void FMaterialBakingModule::BakeMaterials(const TArray<FMaterialDataEx*>& Materi
 
 				TRACE_CPUPROFILER_EVENT_SCOPE_TEXT(*Property.ToString())
 
-				FExportMaterialProxy* ExportMaterialProxy = CreateMaterialProxy(CurrentMaterialSettings->Material, Property);
+				FExportMaterialProxy* ExportMaterialProxy = CreateMaterialProxy(CurrentMaterialSettings, Property);
 
 				if (!ExportMaterialProxy->IsCompilationFinished())
 				{
@@ -873,7 +874,7 @@ UTextureRenderTarget2D* FMaterialBakingModule::CreateRenderTarget(bool bInForceL
 	return RenderTarget;
 }
 
-FExportMaterialProxy* FMaterialBakingModule::CreateMaterialProxy(UMaterialInterface* Material, const FMaterialPropertyEx& Property)
+FExportMaterialProxy* FMaterialBakingModule::CreateMaterialProxy(const FMaterialDataEx* MaterialSettings, const FMaterialPropertyEx& Property)
 {
 	TRACE_CPUPROFILER_EVENT_SCOPE(FMaterialBakingModule::CreateMaterialProxy)
 
@@ -881,12 +882,12 @@ FExportMaterialProxy* FMaterialBakingModule::CreateMaterialProxy(UMaterialInterf
 
 	// Find all pooled material proxy matching this material
 	TArray<FMaterialPoolValue> Entries;
-	MaterialProxyPool.MultiFind(Material, Entries);
+	MaterialProxyPool.MultiFind(MaterialSettings->Material, Entries);
 
 	// Look for the matching property
 	for (FMaterialPoolValue& Entry : Entries)
 	{
-		if (Entry.Key == Property)
+		if (Entry.Key == Property && Entry.Value->bTangentSpaceNormal == MaterialSettings->bTangentSpaceNormal)
 		{
 			Proxy = Entry.Value;
 			break;
@@ -896,8 +897,8 @@ FExportMaterialProxy* FMaterialBakingModule::CreateMaterialProxy(UMaterialInterf
 	// Not found, create a new entry
 	if (Proxy == nullptr)
 	{
-		Proxy = new FExportMaterialProxy(Material, Property.Type, Property.CustomOutput.ToString(), false /* bInSynchronousCompilation */);
-		MaterialProxyPool.Add(Material, FMaterialPoolValue(Property, Proxy));
+		Proxy = new FExportMaterialProxy(MaterialSettings->Material, Property.Type, Property.CustomOutput.ToString(), false /* bInSynchronousCompilation */, MaterialSettings->bTangentSpaceNormal);
+		MaterialProxyPool.Add(MaterialSettings->Material, FMaterialPoolValue(Property, Proxy));
 	}
 
 	return Proxy;
diff --git a/Engine/Source/Developer/MaterialBaking/Private/MaterialBakingModule.h b/Engine/Source/Developer/MaterialBaking/Private/MaterialBakingModule.h
--- a/Engine/Source/Developer/MaterialBaking/Private/MaterialBakingModule.h
+++ b/Engine/Source/Developer/MaterialBaking/Private/MaterialBakingModule.h
@@ -48,7 +48,7 @@ protected:
 	UTextureRenderTarget2D* CreateRenderTarget(bool bInForceLinearGamma, EPixelFormat InPixelFormat, const FIntPoint& InTargetSize);
 
 	/* Creates and adds (or reuses a ExportMaterialProxy from the pool if MaterialBaking.UseMaterialProxyCaching is set to 1) */
-	FExportMaterialProxy* CreateMaterialProxy(UMaterialInterface* Material, const FMaterialPropertyEx& Property);
+	FExportMaterialProxy* CreateMaterialProxy(const FMaterialDataEx* MaterialSettings, const FMaterialPropertyEx& Property);
 
 	/** Helper for emissive color conversion to Output */
 	static void ProcessEmissiveOutput(const FFloat16Color* Color16, int32 Color16Pitch, const FIntPoint& OutputSize, TArray<FColor>& Output, float& EmissiveScale);
diff --git a/Engine/Source/Developer/MaterialBaking/Public/MaterialBakingStructures.h b/Engine/Source/Developer/MaterialBaking/Public/MaterialBakingStructures.h
--- a/Engine/Source/Developer/MaterialBaking/Public/MaterialBakingStructures.h
+++ b/Engine/Source/Developer/MaterialBaking/Public/MaterialBakingStructures.h
@@ -16,6 +16,7 @@ struct FMaterialData
 	FMaterialData()
 		: Material(nullptr)
 		, bPerformBorderSmear(true)
+		, bTangentSpaceNormal(false)
 	{}
 
 	/** Material to bake out */
@@ -24,6 +25,8 @@ struct FMaterialData
 	TMap<EMaterialProperty, FIntPoint> PropertySizes;
 	/** Whether to smear borders after baking */
 	bool bPerformBorderSmear;
+	/** Whether to transform normals from world-space to tangent-space (does nothing if material already uses tangent-space normals) */
+	bool bTangentSpaceNormal;
 };
 
 /** Structure containing extended information about the material and properties which is being baked out */
@@ -32,6 +35,7 @@ struct FMaterialDataEx
 	FMaterialDataEx()
 		: Material(nullptr)
 		, bPerformBorderSmear(true)
+		, bTangentSpaceNormal(false)
 	{}
 
 	/** Material to bake out */
@@ -40,6 +44,8 @@ struct FMaterialDataEx
 	TMap<FMaterialPropertyEx, FIntPoint> PropertySizes;
 	/** Whether to smear borders after baking */
 	bool bPerformBorderSmear;
+	/** Whether to transform normals from world-space to tangent-space (does nothing if material already uses tangent-space normals) */
+	bool bTangentSpaceNormal;
 };
 
 struct FMeshData
