PR #8529: Fix compilation issue when baking materials that use atmospheric fog color (Contributed by [REDACTED_PERSON])

#rnx

diff --git a/Engine/Source/Developer/MaterialBaking/Private/ExportMaterialProxy.h b/Engine/Source/Developer/MaterialBaking/Private/ExportMaterialProxy.h
--- a/Engine/Source/Developer/MaterialBaking/Private/ExportMaterialProxy.h
+++ b/Engine/Source/Developer/MaterialBaking/Private/ExportMaterialProxy.h
@@ -143,11 +143,6 @@ struct FExportMaterialCompiler : public FProxyMaterialCompiler
 		return Compiler->Sub(Compiler->Mul(Compiler->Constant(2.0f), Compiler->Mul(Compiler->Dot(N, C), N)), C);
 	}
 
-	virtual int32 VertexColor() override
-	{
-		return Compiler->VertexColor();
-	}
-
 	virtual int32 PreSkinnedPosition() override
 	{
 		return Compiler->PreSkinnedPosition();
@@ -163,27 +158,12 @@ struct FExportMaterialCompiler : public FProxyMaterialCompiler
 		return Compiler->VertexInterpolator(InterpolatorIndex);
 	}
 
-	virtual int32 LightVector() override
-	{
-		return Compiler->LightVector();
-	}
-
 	virtual int32 ReflectionVector() override
 	{
 		// Because camera vector is identical to normal vector we can work out that reflection vector will also be the same
 		return Compiler->VertexNormal();
 	}
 
-	virtual int32 AtmosphericFogColor(int32 WorldPosition) override
-	{
-		return INDEX_NONE;
-	}
-
-	virtual int32 PrecomputedAOMask() override
-	{
-		return Compiler->PrecomputedAOMask();
-	}
-
 #if WITH_EDITOR
 	virtual int32 MaterialBakingWorldPosition() override
 	{
