PR #8501: Fix support for material baking subsurface & cloth color property (Contributed by [REDACTED_PERSON])

#rnx

diff --git a/Engine/Source/Developer/MaterialBaking/Private/ExportMaterialProxy.h b/Engine/Source/Developer/MaterialBaking/Private/ExportMaterialProxy.h
--- a/Engine/Source/Developer/MaterialBaking/Private/ExportMaterialProxy.h
+++ b/Engine/Source/Developer/MaterialBaking/Private/ExportMaterialProxy.h
@@ -128,16 +128,16 @@ struct FExportMaterialCompiler : public FProxyMaterialCompiler
 		if (CustomWorldNormal == INDEX_NONE)
 		{
 			return INDEX_NONE;
-		}
+	}
 
 		int32 N = CustomWorldNormal;
 		int32 C = CameraVector();
 
 		if (bNormalizeCustomWorldNormal)
-		{
+	{
 			// N = N / sqrt(dot(N, N))
 			N = Compiler->Div(N, Compiler->SquareRoot(Compiler->Dot(N, N)));
-		}
+	}
 
 		// return 2 * dot(N, C) * N - C
 		return Compiler->Sub(Compiler->Mul(Compiler->Constant(2.0f), Compiler->Mul(Compiler->Dot(N, C), N)), C);
@@ -386,6 +386,7 @@ public:
 			case MP_OpacityMask:
 			case MP_CustomData0:
 			case MP_CustomData1:
+			case MP_SubsurfaceColor:
 			{
 				return MaterialInterface->CompileProperty(&ProxyCompiler, PropertyToCompile, ForceCast_Exact_Replicate);
 			}
