Fix crash when subsurface property is being exported upon material export. Happens when nothing is connected to the subsurface pin in the material.

#rnx

diff --git a/Engine/Source/Developer/MaterialBaking/Private/ExportMaterialProxy.h b/Engine/Source/Developer/MaterialBaking/Private/ExportMaterialProxy.h
--- a/Engine/Source/Developer/MaterialBaking/Private/ExportMaterialProxy.h
+++ b/Engine/Source/Developer/MaterialBaking/Private/ExportMaterialProxy.h
@@ -34,12 +34,12 @@ struct FExportMaterialCompiler : public FProxyMaterialCompiler
 
 	virtual FMaterialShadingModelField GetMaterialShadingModels() const override
 	{
-		return MSM_MAX;
+		return Compiler->GetMaterialShadingModels();
 	}
 
 	virtual FMaterialShadingModelField GetCompiledShadingModels() const override
 	{
-		return MSM_MAX;
+		return Compiler->GetCompiledShadingModels();
 	}
 
 	virtual int32 WorldPosition(EWorldPositionIncludedOffsets WorldPositionIncludedOffsets) override
