Added virtual texture warmup for material baking.
Number of warmup frames is controlled by MaterialBaking.VTWarmupFrames CVar and defaults to 5
Warmup should only happen if material samples virtual textures.

diff --git a/Engine/Source/Developer/MaterialBaking/Private/MaterialBakingModule.cpp b/Engine/Source/Developer/MaterialBaking/Private/MaterialBakingModule.cpp
--- a/Engine/Source/Developer/MaterialBaking/Private/MaterialBakingModule.cpp
+++ b/Engine/Source/Developer/MaterialBaking/Private/MaterialBakingModule.cpp
@@ -59,6 +59,12 @@ static TAutoConsoleVariable<int32> CVarMaterialBakingRDOCCapture(
 	TEXT("1: Turned On"),
 	ECVF_Default);
 
+static TAutoConsoleVariable<int32> CVarMaterialBakingVTWarmupFrames(
+	TEXT("MaterialBaking.VTWarmupFrames"),
+	5,
+	TEXT("Number of frames to render for virtual texture warmup when material baking."));
+
+
 namespace FMaterialBakingModuleImpl
 {
 	// Custom dynamic mesh allocator specifically tailored for Material Baking.
@@ -603,12 +609,26 @@ void FMaterialBakingModule::BakeMaterials(const TArray<FMaterialDataEx*>& Materi
 							Canvas.SetRenderTargetRect(FIntRect(0, 0, RenderTarget->GetSurfaceWidth(), RenderTarget->GetSurfaceHeight()));
 							Canvas.SetBaseTransform(Canvas.CalcBaseTransform2D(RenderTarget->GetSurfaceWidth(), RenderTarget->GetSurfaceHeight()));
 
+							// Virtual textures may require repeated rendering to warm up.
+							int32 WarmupIterationCount = 1;
+ 							if (UseVirtualTexturing(ViewFamily.GetFeatureLevel()))
+ 							{
+ 								const FMaterial& MeshMaterial = ExportMaterialProxy->GetIncompleteMaterialWithFallback(ViewFamily.GetFeatureLevel());
+								if (!MeshMaterial.GetUniformVirtualTextureExpressions().IsEmpty())
+								{
+									WarmupIterationCount = CVarMaterialBakingVTWarmupFrames.GetValueOnAnyThread();
+								}
+ 							}
+
 							// Do rendering
-							Canvas.Clear(RenderTarget->ClearColor);
-							FCanvas::FCanvasSortElement& SortElement = Canvas.GetSortElement(Canvas.TopDepthSortKey());
-							SortElement.RenderBatchArray.Add(&RenderItem);
-							Canvas.Flush_RenderThread(RHICmdList);
-							SortElement.RenderBatchArray.Empty();
+							for (int WarmupIndex = 0; WarmupIndex < WarmupIterationCount; ++WarmupIndex)
+							{
+								Canvas.Clear(RenderTarget->ClearColor);
+								FCanvas::FCanvasSortElement& SortElement = Canvas.GetSortElement(Canvas.TopDepthSortKey());
+								SortElement.RenderBatchArray.Add(&RenderItem);
+								Canvas.Flush_RenderThread(RHICmdList);
+								SortElement.RenderBatchArray.Empty();
+							}
 
 							FTexture2DRHIRef StagingBufferRef = StagingBufferPool.CreateStagingBuffer_RenderThread(RHICmdList, RenderTargetResource->GetSizeX(), RenderTargetResource->GetSizeY(), RenderTarget->GetFormat(), RenderTarget->IsSRGB());
 							FGPUFenceRHIRef GPUFence = RHICreateGPUFence(TEXT("MaterialBackingFence"));
