Material Baking - FExportMaterialProxy - Do not rely on BaseMaterialId as it is invalid on cooked materials
* Instead, use StateId & override GetShaderMapUsage() to differentiate the export proxy from the source material.

diff --git a/Engine/Source/Developer/MaterialBaking/Private/ExportMaterialProxy.h b/Engine/Source/Developer/MaterialBaking/Private/ExportMaterialProxy.h
--- a/Engine/Source/Developer/MaterialBaking/Private/ExportMaterialProxy.h
+++ b/Engine/Source/Developer/MaterialBaking/Private/ExportMaterialProxy.h
@@ -243,11 +243,6 @@ public:
 		FMaterialShaderMapId ResourceId;
 		Resource->GetShaderMapId(GMaxRHIShaderPlatform, nullptr, ResourceId);
 
-		// Our Id must be the same as BaseMaterialId for the shader compiler
-		// to be able to set back GameThreadShaderMap after async compilation.
-		Id = ResourceId.BaseMaterialId;
-
-
 		{
 			TArray<FShaderType*> ShaderTypes;
 			TArray<FVertexFactoryType*> VFTypes;
@@ -285,6 +280,9 @@ public:
 			break;
 		};
 
+		Usage = ResourceId.Usage;
+		ResourceId.BaseMaterialId = Material->StateId;
+
 		CacheShaders(ResourceId, GMaxRHIShaderPlatform);
 	}
 
@@ -437,6 +435,12 @@ public:
 		}
 	}
 
+	/**
+	 * Gets the shader map usage of the material, which will be included in the DDC key.
+	 * This mechanism allows derived material classes to create different DDC keys with the same base material.
+	 */
+	virtual EMaterialShaderMapUsage::Type GetShaderMapUsage() const override { return Usage; }
+
 	virtual FString GetMaterialUsageDescription() const override
 	{
 		return FString::Printf(TEXT("MaterialBaking_%s"), MaterialInterface ? *MaterialInterface->GetName() : TEXT("NULL"));
@@ -510,7 +514,14 @@ public:
 	* Should shaders compiled for this material be saved to disk?
 	*/
 	virtual bool IsPersistent() const override { return true; }
-	virtual FGuid GetMaterialId() const override { return Id; }
+
+	virtual FGuid GetMaterialId() const override
+	{
+		// Reuse the base material's Id
+		// Normally this would cause a bug as the shader map would try to be shared by both, 
+		// But FExportMaterialProxy::GetShaderMapUsage() allows this to work
+		return Material->StateId;
+	}
 
 	virtual UMaterialInterface* GetMaterialInterface() const override
 	{
@@ -601,9 +612,10 @@ private:
 	TArray<TObjectPtr<UObject>> ReferencedTextures;
 	/** The property to compile for rendering the sample */
 	EMaterialProperty PropertyToCompile;
+	/** Stores which exported attribute this proxy is compiling for. */
+	EMaterialShaderMapUsage::Type Usage;
 	/** The name of the specific custom output to compile for rendering the sample. Only used if PropertyToCompile is MP_CustomOutput */
 	FString CustomOutputToCompile;
-	FGuid Id;
 	bool bSynchronousCompilation;
 
 public:
