Fix transition errors in staging texture.

diff --git a/Engine/Source/Developer/MaterialBaking/Private/MaterialBakingModule.cpp b/Engine/Source/Developer/MaterialBaking/Private/MaterialBakingModule.cpp
--- a/Engine/Source/Developer/MaterialBaking/Private/MaterialBakingModule.cpp
+++ b/Engine/Source/Developer/MaterialBaking/Private/MaterialBakingModule.cpp
@@ -632,10 +632,7 @@
 
 							FTexture2DRHIRef StagingBufferRef = StagingBufferPool.CreateStagingBuffer_RenderThread(RHICmdList, RenderTargetResource->GetSizeX(), RenderTargetResource->GetSizeY(), RenderTarget->GetFormat(), RenderTarget->IsSRGB());
 							FGPUFenceRHIRef GPUFence = RHICreateGPUFence(TEXT("MaterialBackingFence"));
-
-							RHICmdList.Transition(FRHITransitionInfo(RenderTargetResource->GetRenderTargetTexture(), ERHIAccess::Unknown, ERHIAccess::CopySrc));
-							RHICmdList.CopyTexture(RenderTargetResource->GetRenderTargetTexture(), StagingBufferRef, {});
-							RHICmdList.Transition(FRHITransitionInfo(RenderTargetResource->GetRenderTargetTexture(), ERHIAccess::CopySrc, ERHIAccess::SRVMask));
+							TransitionAndCopyTexture(RHICmdList, RenderTargetResource->GetRenderTargetTexture(), StagingBufferRef, {});
 							RHICmdList.WriteGPUFence(GPUFence);
 
 							// Prepare a lambda for final processing that will be executed asynchronously
