Material Baking - Default initialize FMeshData::TextureCoordinateBox() to (0,0) (1,1)

diff --git a/Engine/Source/Developer/MaterialBaking/Public/MaterialBakingStructures.h b/Engine/Source/Developer/MaterialBaking/Public/MaterialBakingStructures.h
--- a/Engine/Source/Developer/MaterialBaking/Public/MaterialBakingStructures.h
+++ b/Engine/Source/Developer/MaterialBaking/Public/MaterialBakingStructures.h
@@ -13,58 +13,48 @@ struct FMeshDescription;
 /** Structure containing information about the material which is being baked out */
 struct FMaterialData
 {
-	FMaterialData()
-		: Material(nullptr)
-		, bPerformBorderSmear(true)
-		, bTangentSpaceNormal(false)
-	{}
-
 	/** Material to bake out */
-	UMaterialInterface* Material;
+	UMaterialInterface* Material = nullptr;
+
 	/** Properties and the texture size at which they should be baked out */
 	TMap<EMaterialProperty, FIntPoint> PropertySizes;
+
 	/** Whether to smear borders after baking */
-	bool bPerformBorderSmear;
+	bool bPerformBorderSmear = true;
+
 	/** Whether to transform normals from world-space to tangent-space (does nothing if material already uses tangent-space normals) */
-	bool bTangentSpaceNormal;
+	bool bTangentSpaceNormal = false;
 };
 
 /** Structure containing extended information about the material and properties which is being baked out */
 struct FMaterialDataEx
 {
-	FMaterialDataEx()
-		: Material(nullptr)
-		, bPerformBorderSmear(true)
-		, bTangentSpaceNormal(false)
-	{}
-
 	/** Material to bake out */
-	UMaterialInterface* Material;
+	UMaterialInterface* Material = nullptr;
+
 	/** Extended properties and the texture size at which they should be baked out */
 	TMap<FMaterialPropertyEx, FIntPoint> PropertySizes;
+
 	/** Whether to smear borders after baking */
-	bool bPerformBorderSmear;
+	bool bPerformBorderSmear = true;
+
 	/** Whether to transform normals from world-space to tangent-space (does nothing if material already uses tangent-space normals) */
-	bool bTangentSpaceNormal;
+	bool bTangentSpaceNormal = false;
 };
 
 struct FMeshData
 {
-	FMeshData()
-		: MeshDescription(nullptr), Mesh(nullptr), bMirrored(false), VertexColorHash(0), TextureCoordinateIndex(0), LightMapIndex(0), LightMap(nullptr), LightmapResourceCluster(nullptr)
-	{}
-
 	/** Ptr to raw mesh data to use for baking out the material data, if nullptr a standard quad is used */
-	const FMeshDescription* MeshDescription;
+	const FMeshDescription* MeshDescription = nullptr;
 
 	/** Ptr to original static mesh this mesh data came from */
-	const UStaticMesh* Mesh;
+	const UStaticMesh* Mesh = nullptr;
 
 	/** Transform determinant used to detect mirroring */
-	bool bMirrored;
+	bool bMirrored = false;
 
 	/** A hash of the vertex color buffer for the rawmesh */
-	uint32 VertexColorHash;
+	uint32 VertexColorHash = 0;
 
 	/** Material indices to test the Raw Mesh data against, ensuring we only bake out triangles which use the currently baked out material */
 	TArray<int32> MaterialIndices;
@@ -73,28 +63,24 @@ struct FMeshData
 	TArray<FVector2D> CustomTextureCoordinates;
 
 	/** Box which's space contains the UV coordinates used to bake out the material */
-	FBox2D TextureCoordinateBox;
+	FBox2D TextureCoordinateBox = FBox2D(FVector2D(0.0f, 0.0f), FVector2D(1.0f, 1.0f));
 	
 	/** Specific texture coordinate index to use as texture coordinates to bake out the material (is overruled if CustomTextureCoordinates contains any data) */
-	int32 TextureCoordinateIndex;
+	int32 TextureCoordinateIndex = 0;
 
 	/** Light map index used to retrieve the light-map UVs from RawMesh */
-	int32 LightMapIndex;
+	int32 LightMapIndex = 0;
 
 	/** Reference to the lightmap texture part of the level in the currently being baked out mesh instance data is resident */
-	FLightMapRef LightMap;
+	FLightMapRef LightMap = nullptr;
 
 	/** Pointer to the LightmapResourceCluster to be passed on the the LightCacheInterface when baking */
-	const FLightmapResourceCluster* LightmapResourceCluster;
+	const FLightmapResourceCluster* LightmapResourceCluster = nullptr;
 };
 
 /** Structure containing data being processed while baking out materials*/
 struct FBakeOutput
 {
-	FBakeOutput()
-		: EmissiveScale(1.0f)
-	{}
-
 	/** Contains the resulting texture data for baking out a material's property */
 	TMap<EMaterialProperty, TArray<FColor>> PropertyData;
 
@@ -105,16 +91,12 @@ struct FBakeOutput
 	TMap<EMaterialProperty, TArray<FFloat16Color>> HDRPropertyData;
 
 	/** Scale used to allow having wide ranges of emissive values in the source materials, the final proxy material will use this value to scale the emissive texture's pixel values */
-	float EmissiveScale;
+	float EmissiveScale = 1.0f;
 };
 
 /** Structure containing extended data being processed while baking out materials*/
 struct FBakeOutputEx
 {
-	FBakeOutputEx()
-		: EmissiveScale(1.0f)
-	{}
-
 	/** Contains the resulting texture data for baking out a extened material's property */
 	TMap<FMaterialPropertyEx, TArray<FColor>> PropertyData;
 
@@ -125,5 +107,5 @@ struct FBakeOutputEx
 	TMap<FMaterialPropertyEx, TArray<FFloat16Color>> HDRPropertyData;
 
 	/** Scale used to allow having wide ranges of emissive values in the source materials, the final proxy material will use this value to scale the emissive texture's pixel values */
-	float EmissiveScale;
+	float EmissiveScale = 1.0f;
 };
