Add support for material baking expressions that rely on primitive data

(based on PR #8526)

diff --git a/Engine/Source/Developer/MaterialBaking/Private/ExportMaterialProxy.h b/Engine/Source/Developer/MaterialBaking/Private/ExportMaterialProxy.h
--- a/Engine/Source/Developer/MaterialBaking/Private/ExportMaterialProxy.h
+++ b/Engine/Source/Developer/MaterialBaking/Private/ExportMaterialProxy.h
@@ -49,27 +49,17 @@ struct FExportMaterialCompiler : public FProxyMaterialCompiler
 	virtual int32 WorldPosition(EWorldPositionIncludedOffsets WorldPositionIncludedOffsets) override
 	{
 #if WITH_EDITOR
-		return Compiler->MaterialBakingWorldPosition();
+		return MaterialBakingWorldPosition();
 #else
 		return Compiler->WorldPosition(WorldPositionIncludedOffsets);
 #endif
 	}
 
-	virtual int32 ObjectWorldPosition() override
-	{
-		return Compiler->ObjectWorldPosition();
-	}
-
 	virtual int32 DistanceCullFade() override
 	{
 		return Compiler->Constant(1.0f);
 	}
 
-	virtual int32 ActorWorldPosition() override
-	{
-		return Compiler->ActorWorldPosition();
-	}
-
 	virtual int32 ParticleRelativeTime() override
 	{
 		return Compiler->Constant(0.0f);
@@ -106,21 +96,6 @@ struct FExportMaterialCompiler : public FProxyMaterialCompiler
 		return Compiler->Constant2(0.0f, 0.0f);
 	}
 
-	virtual int32 ObjectRadius() override
-	{
-		return Compiler->Constant(500);
-	}
-
-	virtual int32 ObjectBounds() override
-	{
-		return Compiler->ObjectBounds();
-	}
-
-	virtual int32 PreSkinnedLocalBounds(int32 OutputIndex) override
-	{
-		return Compiler->PreSkinnedLocalBounds(OutputIndex);
-	}
-
 	virtual int32 CameraVector() override
 	{
 		// By returning vertex normal instead of a constant vector (like up), we ensure materials (with fresnel for example) are more correctly baked using custom mesh data.
@@ -171,7 +146,9 @@ struct FExportMaterialCompiler : public FProxyMaterialCompiler
 #if WITH_EDITOR
 	virtual int32 MaterialBakingWorldPosition() override
 	{
-		return Compiler->MaterialBakingWorldPosition();
+		// Depending on how the mesh data was retrieved, baking position may only be in local-space
+		const int32 BakingPosition = Compiler->MaterialBakingWorldPosition();
+		return Compiler->TransformPosition(MCB_Local, MCB_World, BakingPosition);
 	}
 #endif
 
diff --git a/Engine/Source/Developer/MaterialBaking/Private/MaterialRenderItem.cpp b/Engine/Source/Developer/MaterialBaking/Private/MaterialRenderItem.cpp
--- a/Engine/Source/Developer/MaterialBaking/Private/MaterialRenderItem.cpp
+++ b/Engine/Source/Developer/MaterialBaking/Private/MaterialRenderItem.cpp
@@ -120,7 +120,20 @@ void FMeshMaterialRenderItem::QueueMaterial(FCanvasRenderContext& RenderContext,
 			DynamicMeshBuilder.AddTriangles(Indices);
 		}
 
-		DynamicMeshBuilder.GetMeshElement(FMatrix::Identity, MaterialRenderProxy, SDPG_Foreground, true, false, 0, MeshBuilderResources, MeshElement);
+		FPrimitiveData PrimitiveData = MeshSettings->PrimitiveData != nullptr ? *MeshSettings->PrimitiveData : FPrimitiveData();
+		FPrimitiveUniformShaderParameters PrimitiveParams = FPrimitiveUniformShaderParametersBuilder{}
+			.Defaults()
+				.LocalToWorld(PrimitiveData.LocalToWorld)
+				.ActorWorldPosition(PrimitiveData.ActorPosition)
+				.WorldBounds(PrimitiveData.WorldBounds)
+				.LocalBounds(PrimitiveData.LocalBounds)
+				.PreSkinnedLocalBounds(PrimitiveData.PreSkinnedLocalBounds)
+				.CustomPrimitiveData(PrimitiveData.CustomPrimitiveData)
+				.ReceivesDecals(false)
+				.OutputVelocity(false)
+			.Build();
+
+		DynamicMeshBuilder.GetMeshElement(PrimitiveParams, MaterialRenderProxy, SDPG_Foreground, true, 0, MeshBuilderResources, MeshElement);
 
 		check(MeshBuilderResources.IsValidForRendering());
 		bMeshElementDirty = false;
@@ -148,6 +161,9 @@ void FMeshMaterialRenderItem::QueueMaterial(FCanvasRenderContext& RenderContext,
 
 void FMeshMaterialRenderItem::PopulateWithQuadData()
 {
+	// Pre-transform all vertices with the inverse of LocalToWorld to negate its effect during material baking 
+	const FMatrix44f WorldToLocal = (MeshSettings->PrimitiveData != nullptr) ? FMatrix44f(MeshSettings->PrimitiveData->LocalToWorld.Inverse()) : FMatrix44f::Identity;
+	
 	Vertices.Empty(4);
 	Indices.Empty(6);
 
@@ -164,13 +180,16 @@ void FMeshMaterialRenderItem::PopulateWithQuadData()
 		FDynamicMeshVertex* Vert = new(Vertices)FDynamicMeshVertex();
 		const int32 X = VertIndex & 1;
 		const int32 Y = (VertIndex >> 1) & 1;
-		Vert->Position.Set(ScaleX * X, ScaleY * Y, 0);
-		Vert->SetTangents(FVector3f(1, 0, 0), FVector3f(0, 1, 0), FVector3f(0, 0, 1));
+		Vert->Position = WorldToLocal.TransformPosition(FVector3f(ScaleX * X, ScaleY * Y, 0));
+		FVector3f TangentX = WorldToLocal.TransformVector(FVector3f(1, 0, 0));
+		FVector3f TangentZ = WorldToLocal.TransformVector(FVector3f(0, 1, 0));
+		FVector3f TangentY = WorldToLocal.TransformVector(FVector3f(0, 0, 1));
+		Vert->SetTangents(TangentX, TangentZ, TangentY);
 		FMemory::Memzero(&Vert->TextureCoordinate, sizeof(Vert->TextureCoordinate));
 		for (int32 TexcoordIndex = 0; TexcoordIndex < MAX_STATIC_TEXCOORDS; TexcoordIndex++)
 		{
 			Vert->TextureCoordinate[TexcoordIndex].Set(U + SizeU * X, V + SizeV * Y);
-		}		
+		}
 		Vert->Color = FColor::White;
 	}
 
@@ -181,6 +200,9 @@ void FMeshMaterialRenderItem::PopulateWithQuadData()
 
 void FMeshMaterialRenderItem::PopulateWithMeshData()
 {
+	// Pre-transform all vertices with the inverse of LocalToWorld to negate its effect during material baking 
+	const FMatrix44f WorldToLocal = (MeshSettings->PrimitiveData != nullptr) ? FMatrix44f(MeshSettings->PrimitiveData->LocalToWorld.Inverse()) : FMatrix44f::Identity;
+
 	const FMeshDescription* RawMesh = MeshSettings->MeshDescription;
 
 	FStaticMeshConstAttributes Attributes(*RawMesh);
@@ -205,7 +227,7 @@ void FMeshMaterialRenderItem::PopulateWithMeshData()
 	const int32 NumTexcoords = [&]()
 	{
 		return FMath::Min(VertexInstanceUVs.GetNumChannels(), VertexPositionStoredUVChannel);
-	}();		
+	}();
 
 	// check if we should use NewUVs or original UV set
 	const bool bUseNewUVs = MeshSettings->CustomTextureCoordinates.Num() > 0;
@@ -238,22 +260,22 @@ void FMeshMaterialRenderItem::PopulateWithMeshData()
 				{
 					// compute vertex position from original UV
 					const FVector2D& UV = FVector2D(VertexInstanceUVs.Get(SrcVertexInstanceID, MeshSettings->TextureCoordinateIndex));
-					Vert->Position.Set(UV.X * ScaleX, UV.Y * ScaleY, 0);
+					Vert->Position = WorldToLocal.TransformPosition(FVector3f(UV.X * ScaleX, UV.Y * ScaleY, 0));
 				}
 				else
 				{
 					const FVector2D& UV = MeshSettings->CustomTextureCoordinates[SrcVertIndex];
-					Vert->Position.Set(UV.X * ScaleX, UV.Y * ScaleY, 0);
+					Vert->Position = WorldToLocal.TransformPosition(FVector3f(UV.X * ScaleX, UV.Y * ScaleY, 0));
 				}
-				FVector3f TangentX = VertexInstanceTangents[SrcVertexInstanceID];
-				FVector3f TangentZ = VertexInstanceNormals[SrcVertexInstanceID];
+				FVector3f TangentX = WorldToLocal.TransformVector(VertexInstanceTangents[SrcVertexInstanceID]);
+				FVector3f TangentZ = WorldToLocal.TransformVector(VertexInstanceNormals[SrcVertexInstanceID]);
 				FVector3f TangentY = FVector3f::CrossProduct(TangentZ, TangentX).GetSafeNormal() * VertexInstanceBinormalSigns[SrcVertexInstanceID];
 				Vert->SetTangents(TangentX, TangentY, TangentZ);
 				for (int32 TexcoordIndex = 0; TexcoordIndex < NumTexcoords; TexcoordIndex++)
 				{
 					Vert->TextureCoordinate[TexcoordIndex] = VertexInstanceUVs.Get(SrcVertexInstanceID, TexcoordIndex);
 				}
-				
+
 				if (NumTexcoords < VertexPositionStoredUVChannel)
 				{
 					for (int32 TexcoordIndex = NumTexcoords; TexcoordIndex < VertexPositionStoredUVChannel; TexcoordIndex++)
diff --git a/Engine/Source/Developer/MaterialBaking/Public/MaterialBakingStructures.h b/Engine/Source/Developer/MaterialBaking/Public/MaterialBakingStructures.h
--- a/Engine/Source/Developer/MaterialBaking/Public/MaterialBakingStructures.h
+++ b/Engine/Source/Developer/MaterialBaking/Public/MaterialBakingStructures.h
@@ -3,6 +3,10 @@
 #pragma once
 
 #include "Math/IntPoint.h"
+#include "Components/MeshComponent.h"
+#include "Components/SkinnedMeshComponent.h"
+#include "Engine/SkeletalMesh.h"
+#include "Engine/StaticMesh.h"
 #include "SceneTypes.h"
 #include "LightMap.h"
 #include "MaterialPropertyEx.h"
@@ -48,6 +52,57 @@ struct FMaterialDataEx
 	bool bTangentSpaceNormal = false;
 };
 
+/** Structure containing primitive information (regarding a mesh or mesh component) that is accessible through material expressions */
+struct FPrimitiveData
+{
+	FPrimitiveData(const FBoxSphereBounds& LocalBounds = FBoxSphereBounds(ForceInitToZero))
+		: LocalToWorld(FMatrix::Identity)
+		, ActorPosition(ForceInitToZero)
+		, WorldBounds(LocalBounds)
+		, LocalBounds(LocalBounds)
+		, PreSkinnedLocalBounds(LocalBounds)
+		, CustomPrimitiveData(nullptr)
+	{}
+
+	FPrimitiveData(const UStaticMesh* StaticMesh)
+		: FPrimitiveData(StaticMesh->GetBounds())
+	{}
+
+	FPrimitiveData(const USkeletalMesh* SkeletalMesh)
+		: FPrimitiveData(SkeletalMesh->GetBounds())
+	{}
+
+	FPrimitiveData(const UMeshComponent* MeshComponent)
+		: LocalToWorld(MeshComponent->GetRenderMatrix())
+		, ActorPosition(MeshComponent->GetActorPositionForRenderer())
+		, WorldBounds(MeshComponent->Bounds)
+		, LocalBounds(MeshComponent->GetLocalBounds())
+		, CustomPrimitiveData(&MeshComponent->GetCustomPrimitiveData())
+	{
+		if (const USkinnedMeshComponent* SkinnedMeshComponent = Cast<USkinnedMeshComponent>(MeshComponent))
+		{
+			SkinnedMeshComponent->GetPreSkinnedLocalBounds(PreSkinnedLocalBounds);
+		}
+		else
+		{
+			PreSkinnedLocalBounds = LocalBounds;
+		}
+	}
+
+	/** The mesh component's local-to-world transform */
+	FMatrix LocalToWorld;
+	/** The actor's location in world-space */
+	FVector ActorPosition;
+	/** The mesh component's bounds in world-space */
+	FBoxSphereBounds WorldBounds;
+	/** The mesh component's bounds in local-space */
+	FBoxSphereBounds LocalBounds;
+	/** The mesh component's pre-skinning bounds in local-space */
+	FBoxSphereBounds PreSkinnedLocalBounds;
+	/** The mesh component's custom primitive data */
+	const FCustomPrimitiveData* CustomPrimitiveData;
+};
+
 struct FMeshData
 {
 	/** Ptr to raw mesh data to use for baking out the material data, if nullptr a standard quad is used */
@@ -82,6 +137,9 @@ struct FMeshData
 
 	/** Pointer to the LightmapResourceCluster to be passed on the the LightCacheInterface when baking */
 	const FLightmapResourceCluster* LightmapResourceCluster = nullptr;
+
+	/** Pointer to primitive data that is accessible through material expressions, if nullptr default values are used */
+	const FPrimitiveData* PrimitiveData = nullptr;
 };
 
 /** Structure containing data being processed while baking out materials*/
diff --git a/Engine/Source/Runtime/Engine/Private/DynamicMeshBuilder.cpp b/Engine/Source/Runtime/Engine/Private/DynamicMeshBuilder.cpp
--- a/Engine/Source/Runtime/Engine/Private/DynamicMeshBuilder.cpp
+++ b/Engine/Source/Runtime/Engine/Private/DynamicMeshBuilder.cpp
@@ -911,6 +911,21 @@ void FDynamicMeshBuilder::GetMesh(
 }
 
 void FDynamicMeshBuilder::GetMeshElement(const FMatrix& LocalToWorld, const FMaterialRenderProxy* MaterialRenderProxy, uint8 DepthPriorityGroup, bool bDisableBackfaceCulling, bool bReceivesDecals, int32 ViewIndex, FMeshBuilderOneFrameResources& OneFrameResource, FMeshBatch& Mesh)
+{
+	FPrimitiveUniformShaderParameters PrimitiveParams = FPrimitiveUniformShaderParametersBuilder{}
+		.Defaults()
+			.LocalToWorld(LocalToWorld)
+			.ActorWorldPosition(LocalToWorld.GetOrigin())
+			.WorldBounds(FBoxSphereBounds(EForceInit::ForceInit))
+			.LocalBounds(FBoxSphereBounds(EForceInit::ForceInit))
+			.ReceivesDecals(bReceivesDecals)
+			.OutputVelocity(true)
+		.Build();
+
+	GetMeshElement(PrimitiveParams, MaterialRenderProxy, DepthPriorityGroup, bDisableBackfaceCulling, ViewIndex, OneFrameResource, Mesh);
+}
+
+void FDynamicMeshBuilder::GetMeshElement(const FPrimitiveUniformShaderParameters& PrimitiveParams, const FMaterialRenderProxy* MaterialRenderProxy, uint8 DepthPriorityGroup, bool bDisableBackfaceCulling, int32 ViewIndex, FMeshBuilderOneFrameResources& OneFrameResource, FMeshBatch& Mesh)
 {
 	TRACE_CPUPROFILER_EVENT_SCOPE(FDynamicMeshBuilder::GetMeshElement)
 
@@ -936,15 +951,6 @@ void FDynamicMeshBuilder::GetMeshElement(const FMatrix& LocalToWorld, const FMat
 
 			// Create the primitive uniform buffer.
 			OneFrameResource.PrimitiveUniformBuffer = new FDynamicMeshPrimitiveUniformBuffer();
-			FPrimitiveUniformShaderParameters PrimitiveParams = FPrimitiveUniformShaderParametersBuilder{}
-				.Defaults()
-					.LocalToWorld(LocalToWorld)
-					.ActorWorldPosition(LocalToWorld.GetOrigin())
-					.WorldBounds(FBoxSphereBounds(EForceInit::ForceInit))
-					.LocalBounds(FBoxSphereBounds(EForceInit::ForceInit))
-					.ReceivesDecals(bReceivesDecals)
-					.OutputVelocity(true)
-				.Build();
 
 			if (IsInGameThread())
 			{
@@ -976,7 +982,7 @@ void FDynamicMeshBuilder::GetMeshElement(const FMatrix& LocalToWorld, const FMat
 		BatchElement.NumPrimitives = bHasValidIndexBuffer ? (OneFrameResource.IndexBuffer->Indices.Num() / 3) : (bHasValidVertexBuffer ? OneFrameResource.VertexBuffer->Vertices.Num() / 3 : 0);
 		BatchElement.MinVertexIndex = 0;
 		BatchElement.MaxVertexIndex = bHasValidVertexBuffer ? OneFrameResource.VertexBuffer->Vertices.Num() - 1 : 0;
-		Mesh.ReverseCulling = LocalToWorld.Determinant() < 0.0f ? true : false;
+		Mesh.ReverseCulling = PrimitiveParams.LocalToRelativeWorld.Determinant() < 0.0f ? true : false;
 		Mesh.bDisableBackfaceCulling = bDisableBackfaceCulling;
 		Mesh.Type = PT_TriangleList;
 		Mesh.DepthPriorityGroup = DepthPriorityGroup;
diff --git a/Engine/Source/Runtime/Engine/Public/DynamicMeshBuilder.h b/Engine/Source/Runtime/Engine/Public/DynamicMeshBuilder.h
--- a/Engine/Source/Runtime/Engine/Public/DynamicMeshBuilder.h
+++ b/Engine/Source/Runtime/Engine/Public/DynamicMeshBuilder.h
@@ -20,6 +20,7 @@ class FMaterialRenderProxy;
 class FMeshBuilderOneFrameResources;
 class FMeshElementCollector;
 class FPrimitiveDrawInterface;
+class FPrimitiveUniformShaderParameters;
 struct FMeshBatch;
 
 /** The vertex type used for dynamic meshes. */
@@ -197,6 +198,7 @@ public:
 		FDynamicMeshDrawOffset const * const DrawOffset, int32 ViewIndex, FMeshElementCollector& Collector, const FHitProxyId HitProxyId = FHitProxyId());
 
 	ENGINE_API void GetMeshElement(const FMatrix& LocalToWorld, const FMaterialRenderProxy* MaterialRenderProxy, uint8 DepthPriorityGroup, bool bDisableBackfaceCulling, bool bReceivesDecals, int32 ViewIndex, FMeshBuilderOneFrameResources& OneFrameResource, FMeshBatch& Mesh);
+	ENGINE_API void GetMeshElement(const FPrimitiveUniformShaderParameters& PrimitiveParams, const FMaterialRenderProxy* MaterialRenderProxy, uint8 DepthPriorityGroup, bool bDisableBackfaceCulling, int32 ViewIndex, FMeshBuilderOneFrameResources& OneFrameResource, FMeshBatch& Mesh);
 
 	/**
 	 * Draws the mesh to the given primitive draw interface.
