Add support for material baking shading model property

(based on PR #8497)

diff --git a/Engine/Source/Developer/MaterialBaking/Private/ExportMaterialProxy.h b/Engine/Source/Developer/MaterialBaking/Private/ExportMaterialProxy.h
--- a/Engine/Source/Developer/MaterialBaking/Private/ExportMaterialProxy.h
+++ b/Engine/Source/Developer/MaterialBaking/Private/ExportMaterialProxy.h
@@ -250,6 +250,7 @@ public:
 		case MP_Opacity: ResourceId.Usage = EMaterialShaderMapUsage::MaterialExportOpacity; break;
 		case MP_OpacityMask: ResourceId.Usage = EMaterialShaderMapUsage::MaterialExportOpacityMask; break;
 		case MP_SubsurfaceColor: ResourceId.Usage = EMaterialShaderMapUsage::MaterialExportSubSurfaceColor; break;
+		case MP_ShadingModel: ResourceId.Usage = EMaterialShaderMapUsage::MaterialExportShadingModel; break;
 		case MP_CustomData0: ResourceId.Usage = EMaterialShaderMapUsage::MaterialExportClearCoat; break;
 		case MP_CustomData1: ResourceId.Usage = EMaterialShaderMapUsage::MaterialExportClearCoatRoughness; break;
 		case MP_CustomOutput:
@@ -379,7 +380,7 @@ public:
 				}
 				break;
 			case MP_ShadingModel:
-				return MaterialInterface->CompileProperty(&ProxyCompiler, MP_ShadingModel);
+				return CompileShadingModelEncoding(Compiler, MaterialInterface->CompileProperty(&ProxyCompiler, MP_ShadingModel));
 			case MP_CustomOutput:
 				 // NOTE: Currently we can assume input index is always 0, which it is for all custom outputs that are registered as material attributes
 				return CompileInputForCustomOutput(&ProxyCompiler, 0, ForceCast_Exact_Replicate);
@@ -602,6 +603,14 @@ private:
 			Compiler->Constant(0.5f)); // [-0.5,0.5] + 0.5
 	}
 
+	static int32 CompileShadingModelEncoding(FMaterialCompiler* Compiler, int32 ShadingModelInput)
+	{
+		// [0,MSM_NUM] -> [0,1]
+		return Compiler->Div(
+			Compiler->CastShadingModelToFloat(ShadingModelInput),
+			Compiler->Constant(MSM_NUM));
+	}
+
 private:
 	/** The material interface for this proxy */
 	UMaterialInterface* MaterialInterface;
diff --git a/Engine/Source/Developer/MaterialBaking/Private/MaterialBakingModule.cpp b/Engine/Source/Developer/MaterialBaking/Private/MaterialBakingModule.cpp
--- a/Engine/Source/Developer/MaterialBaking/Private/MaterialBakingModule.cpp
+++ b/Engine/Source/Developer/MaterialBaking/Private/MaterialBakingModule.cpp
@@ -289,6 +289,7 @@ void FMaterialBakingModule::StartupModule()
 	PerPropertyFormat.Add(MP_SubsurfaceColor, PF_B8G8R8A8);
 	PerPropertyFormat.Add(MP_CustomData0, PF_B8G8R8A8);
 	PerPropertyFormat.Add(MP_CustomData1, PF_B8G8R8A8);
+	PerPropertyFormat.Add(MP_ShadingModel, PF_B8G8R8A8);
 	PerPropertyFormat.Add(TEXT("ClearCoatBottomNormal"), PF_B8G8R8A8);
 
 	// Register property customization
@@ -845,6 +846,7 @@ void FMaterialBakingModule::SetLinearBake(bool bCorrectLinear)
 		PerPropertyColorSpace.Add(MP_Normal, EPropertyColorSpace::Linear);
 		PerPropertyColorSpace.Add(MP_Opacity, EPropertyColorSpace::Linear);
 		PerPropertyColorSpace.Add(MP_OpacityMask, EPropertyColorSpace::Linear);
+		PerPropertyColorSpace.Add(MP_ShadingModel, EPropertyColorSpace::Linear);
 		PerPropertyColorSpace.Add(TEXT("ClearCoatBottomNormal"), EPropertyColorSpace::Linear);
 	}
 }
diff --git a/Engine/Source/Runtime/Engine/Private/Materials/HLSLMaterialTranslator.cpp b/Engine/Source/Runtime/Engine/Private/Materials/HLSLMaterialTranslator.cpp
--- a/Engine/Source/Runtime/Engine/Private/Materials/HLSLMaterialTranslator.cpp
+++ b/Engine/Source/Runtime/Engine/Private/Materials/HLSLMaterialTranslator.cpp
@@ -4140,6 +4140,17 @@ int32 FHLSLMaterialTranslator::ForceCast(int32 Code, EMaterialValueType DestType
 	}
 }
 
+int32 FHLSLMaterialTranslator::CastShadingModelToFloat(int32 Code)
+{
+	const EMaterialValueType SourceType = GetParameterType(Code);
+	if (SourceType != MCT_ShadingModel)
+	{
+		return Errorf(TEXT("Operation only supported for shading model"));
+	}
+
+	return AddCodeChunk(MCT_Float1, TEXT("((%s)(%s))"), HLSLTypeString(MCT_Float1), *GetParameterCode(Code));
+}
+
 int32 FHLSLMaterialTranslator::TruncateLWC(int32 Code)
 {
 	const int32 LWCTruncateMode = GetLWCTruncateMode();
diff --git a/Engine/Source/Runtime/Engine/Private/Materials/HLSLMaterialTranslator.h b/Engine/Source/Runtime/Engine/Private/Materials/HLSLMaterialTranslator.h
--- a/Engine/Source/Runtime/Engine/Private/Materials/HLSLMaterialTranslator.h
+++ b/Engine/Source/Runtime/Engine/Private/Materials/HLSLMaterialTranslator.h
@@ -740,6 +740,8 @@ protected:
 	virtual int32 ValidCast(int32 Code, EMaterialValueType DestType) override;
 	virtual int32 ForceCast(int32 Code, EMaterialValueType DestType, uint32 ForceCastFlags = 0) override;
 
+	virtual int32 CastShadingModelToFloat(int32 Code) override;
+
 	virtual int32 TruncateLWC(int32 Code) override;
 
 	/** Pushes a function onto the compiler's function stack, which indicates that compilation is entering a function. */
diff --git a/Engine/Source/Runtime/Engine/Public/MaterialCompiler.h b/Engine/Source/Runtime/Engine/Public/MaterialCompiler.h
--- a/Engine/Source/Runtime/Engine/Public/MaterialCompiler.h
+++ b/Engine/Source/Runtime/Engine/Public/MaterialCompiler.h
@@ -271,6 +271,9 @@ public:
 	virtual int32 ValidCast(int32 Code,EMaterialValueType DestType) = 0;
 	virtual int32 ForceCast(int32 Code,EMaterialValueType DestType,uint32 ForceCastFlags = 0) = 0;
 
+	/** Cast shading model integer to float value */
+	virtual int32 CastShadingModelToFloat(int32 Code) = 0;
+
 	virtual int32 TruncateLWC(int32 Code) = 0;
 
 	/** Pushes a function onto the compiler's function stack, which indicates that compilation is entering a function. */
@@ -802,6 +805,7 @@ public:
 	virtual bool IsMaterialPropertyUsed(EMaterialProperty Property, int32 CodeChunkIdx) const override { return Compiler->IsMaterialPropertyUsed(Property, CodeChunkIdx); }
 	virtual int32 ValidCast(int32 Code, EMaterialValueType DestType) override { return Compiler->ValidCast(Code, DestType); }
 	virtual int32 ForceCast(int32 Code, EMaterialValueType DestType, uint32 ForceCastFlags = 0) override { return Compiler->ForceCast(Code, DestType, ForceCastFlags); }
+	virtual int32 CastShadingModelToFloat(int32 Code) override { return Compiler->CastShadingModelToFloat(Code); }
 	virtual int32 TruncateLWC(int32 Code) override { return Compiler->TruncateLWC(Code); }
 
 	virtual int32 AccessCollectionParameter(UMaterialParameterCollection* ParameterCollection, int32 ParameterIndex, int32 ComponentIndex) override { return Compiler->AccessCollectionParameter(ParameterCollection, ParameterIndex, ComponentIndex); }
diff --git a/Engine/Source/Runtime/Engine/Public/MaterialShared.h b/Engine/Source/Runtime/Engine/Public/MaterialShared.h
--- a/Engine/Source/Runtime/Engine/Public/MaterialShared.h
+++ b/Engine/Source/Runtime/Engine/Public/MaterialShared.h
@@ -796,6 +796,7 @@ namespace EMaterialShaderMapUsage
 		MaterialExportClearCoat,
 		MaterialExportClearCoatRoughness,
 		MaterialExportCustomOutput,
+		MaterialExportShadingModel,
 	};
 }
 
