Add support for using material baking to efficiently check UV overlap

(based on PR #8500)

diff --git a/Engine/Source/Developer/MaterialBaking/Private/ExportMaterialProxy.h b/Engine/Source/Developer/MaterialBaking/Private/ExportMaterialProxy.h
--- a/Engine/Source/Developer/MaterialBaking/Private/ExportMaterialProxy.h
+++ b/Engine/Source/Developer/MaterialBaking/Private/ExportMaterialProxy.h
@@ -206,7 +206,7 @@ struct FExportMaterialCompiler : public FProxyMaterialCompiler
 class FExportMaterialProxy : public FMaterial, public FMaterialRenderProxy
 {
 public:
-	FExportMaterialProxy(UMaterialInterface* InMaterialInterface, EMaterialProperty InPropertyToCompile, const FString& InCustomOutputToCompile = TEXT(""), bool bInSynchronousCompilation = true, bool bTangentSpaceNormal = false)
+	FExportMaterialProxy(UMaterialInterface* InMaterialInterface, EMaterialProperty InPropertyToCompile, const FString& InCustomOutputToCompile = TEXT(""), bool bInSynchronousCompilation = true, bool bTangentSpaceNormal = false, EBlendMode ProxyBlendMode = BLEND_Opaque)
 		: FMaterial()
 		, FMaterialRenderProxy(GetPathNameSafe(InMaterialInterface->GetMaterial()))
 		, MaterialInterface(InMaterialInterface)
@@ -214,6 +214,7 @@ public:
 		, CustomOutputToCompile(InCustomOutputToCompile)
 		, bSynchronousCompilation(bInSynchronousCompilation)
 		, bTangentSpaceNormal(bTangentSpaceNormal)
+		, ProxyBlendMode(ProxyBlendMode)
 	{
 		SetQualityLevelProperties(GMaxRHIFeatureLevel);
 		Material = InMaterialInterface->GetMaterial();
@@ -494,8 +495,8 @@ public:
 		}
 		return false;
 	}
-	virtual bool IsMasked() const override { return false; }
-	virtual enum EBlendMode GetBlendMode() const override { return BLEND_Opaque; }
+	virtual bool IsMasked() const override { return ProxyBlendMode == BLEND_Masked; }
+	virtual enum EBlendMode GetBlendMode() const override { return ProxyBlendMode; }
 	virtual enum ERefractionMode GetRefractionMode() const override { return Material ? (ERefractionMode)Material->RefractionMethod : RM_None; }
 	virtual bool GetRootNodeOverridesDefaultRefraction()const override { return Material ? Material->bRootNodeOverridesDefaultDistortion : false; }
 	virtual FMaterialShadingModelField GetShadingModels() const override { return MSM_DefaultLit; }
@@ -627,4 +628,6 @@ private:
 public:
 	/** Whether to transform normals from world-space to tangent-space (does nothing if material already uses tangent-space normals) */
 	bool bTangentSpaceNormal;
-};
\ No newline at end of file
+	/** The blend mode used when baking the proxy material */
+	EBlendMode ProxyBlendMode;
+};
diff --git a/Engine/Source/Developer/MaterialBaking/Private/MaterialBakingHelpers.cpp b/Engine/Source/Developer/MaterialBaking/Private/MaterialBakingHelpers.cpp
--- a/Engine/Source/Developer/MaterialBaking/Private/MaterialBakingHelpers.cpp
+++ b/Engine/Source/Developer/MaterialBaking/Private/MaterialBakingHelpers.cpp
@@ -6,18 +6,16 @@
 
 namespace FMaterialBakingHelpersImpl
 {
-	static FColor BoxBlurSample(FColor* InBMP, int32 X, int32 Y, int32 InImageWidth, int32 InImageHeight)
+	static FColor BoxBlurSample(FColor* InBMP, int32 X, int32 Y, int32 InImageWidth, int32 InImageHeight, uint32 BackgroundDWColor)
 	{
 		int32 PixelsSampled = 0;
 		int32 CombinedColorR = 0;
 		int32 CombinedColorG = 0;
 		int32 CombinedColorB = 0;
 		int32 CombinedColorA = 0;
-	
-		const uint32 MagentaMask = FColor(255, 0, 255).DWColor();
 
 		const int32 BaseIndex = (Y * InImageWidth) + X;
-	
+
 		int32 SampleIndex = BaseIndex - (InImageWidth + 1);
 		for (int32 YI = 0; YI < 3; ++YI)
 		{
@@ -25,7 +23,7 @@ namespace FMaterialBakingHelpersImpl
 			{
 				const FColor SampledColor = InBMP[SampleIndex++];
 				// Check if the pixel is a rendered one (not clear color)
-				if (SampledColor.DWColor() != MagentaMask)
+				if (SampledColor.DWColor() != BackgroundDWColor)
 				{
 					CombinedColorR += SampledColor.R;
 					CombinedColorG += SampledColor.G;
@@ -39,17 +37,14 @@ namespace FMaterialBakingHelpersImpl
 
 		if (PixelsSampled == 0)
 		{
-			const FColor Magenta = FColor(255, 0, 255);
-			return Magenta;
+			return FColor(BackgroundDWColor);
 		}
 
 		return FColor(CombinedColorR / PixelsSampled, CombinedColorG / PixelsSampled, CombinedColorB / PixelsSampled, CombinedColorA / PixelsSampled);
 	}
 
-	static bool HasBorderingPixel(FColor* InBMP, int32 X, int32 Y, int32 InImageWidth, int32 InImageHeight)
+	static bool HasBorderingPixel(FColor* InBMP, int32 X, int32 Y, int32 InImageWidth, int32 InImageHeight, uint32 BackgroundDWColor)
 	{
-		const uint32 MagentaMask = FColor(255, 0, 255).DWColor();
-
 		const int32 BaseIndex = Y*InImageWidth + X;
 
 		int32 SampleIndex = BaseIndex - (InImageWidth + 1);
@@ -57,7 +52,7 @@ namespace FMaterialBakingHelpersImpl
 		{
 			for (int32 XI = 0; XI < 3; ++XI)
 			{
-				if (InBMP[SampleIndex++].DWColor() != MagentaMask)
+				if (InBMP[SampleIndex++].DWColor() != BackgroundDWColor)
 				{
 					return true;
 				}
@@ -67,28 +62,26 @@ namespace FMaterialBakingHelpersImpl
 		return false;
 	}
 
-	static void PerformShrinking(TArray<FColor>& InOutPixels, int32& ImageWidth, int32& ImageHeight)
+	static void PerformShrinking(TArray<FColor>& InOutPixels, int32& ImageWidth, int32& ImageHeight, uint32 BackgroundDWColor)
 	{
-		const uint32 MagentaMask = FColor(255, 0, 255).DWColor();
-
 		FColor* Current = InOutPixels.GetData();
 
 		// We can hopefully skip the entire smearing process if there is a single
-		// non-magenta color in the entire array since smearing that completely
+		// non-background color in the entire array since smearing that completely
 		// would lead to a monochome output.
-		uint32 SingleColor = MagentaMask;
+		uint32 SingleColor = BackgroundDWColor;
 		bool   bHasSingleColor = true;
 		for (int32 Index = 0, Num = InOutPixels.Num(); Index < Num; ++Index)
 		{
 			const uint32 Color = Current[Index].DWColor();
-			if (Color != MagentaMask)
+			if (Color != BackgroundDWColor)
 			{
-				if (SingleColor == MagentaMask)
+				if (SingleColor == BackgroundDWColor)
 				{
-					// This is the first time we stumble on a color other than magenta, keep it.
+					// This is the first time we stumble on a color other than background color, keep it.
 					SingleColor = Color;
 				}
-				// Compare with the known non-magenta color
+				// Compare with the known non-background color
 				else if (SingleColor != Color)
 				{
 					bHasSingleColor = false;
@@ -108,7 +101,7 @@ namespace FMaterialBakingHelpersImpl
 		}
 	}
 
-	static void PerformUVBorderSmear(TArray<FColor>& InOutPixels, int32& ImageWidth, int32& ImageHeight, int32 MaxIterations)
+	static void PerformUVBorderSmear(TArray<FColor>& InOutPixels, int32& ImageWidth, int32& ImageHeight, int32 MaxIterations, uint32 BackgroundDWColor)
 	{
 		TRACE_CPUPROFILER_EVENT_SCOPE(FMaterialBakingHelpers::PerformUVBorderSmear)
 
@@ -120,7 +113,6 @@ namespace FMaterialBakingHelpersImpl
 			return;
 		}
 
-		const uint32 MagentaMask = FColor(255, 0, 255).DWColor();
 		const int32 PaddedImageWidth = ImageWidth + 2;
 		const int32 PaddedImageHeight = ImageHeight + 2;
 
@@ -130,7 +122,7 @@ namespace FMaterialBakingHelpersImpl
 		TArray<uint32>* NextRowsLeft = &RowsSwap2;
 
 		//
-		// Create a copy of the image with a 1 pixel magenta border around the edge.
+		// Create a copy of the image with a 1 pixel background color border around the edge.
 		// This avoids needing to bounds check in inner loops.
 		//
 
@@ -138,27 +130,27 @@ namespace FMaterialBakingHelpersImpl
 		ScratchBuffer.SetNumUninitialized(PaddedImageWidth*PaddedImageHeight);
 		FColor* Scratch = ScratchBuffer.GetData();
 
-		// Set top and bottom bordering pixels to magenta
+		// Set top and bottom bordering pixels to background color
 		for (int32 X = 0; X < PaddedImageWidth; ++X)
 		{
-			Scratch[X] = FColor(MagentaMask);
-			Scratch[(PaddedImageHeight - 1)*PaddedImageWidth + X] = FColor(MagentaMask);
+			Scratch[X] = FColor(BackgroundDWColor);
+			Scratch[(PaddedImageHeight - 1)*PaddedImageWidth + X] = FColor(BackgroundDWColor);
 		}
 
 		FColor* Current = InOutPixels.GetData();
 
-		// Set left and right border pixels to magenta and copy image data into rows
+		// Set left and right border pixels to background color and copy image data into rows
 		for (int32 Y = 1; Y <= ImageHeight; ++Y)
 		{
 			const int32 YOffset = Y * PaddedImageWidth;
-			Scratch[YOffset] = FColor(MagentaMask);
+			Scratch[YOffset] = FColor(BackgroundDWColor);
 			FMemory::Memcpy(&Scratch[YOffset + 1], &Current[(Y - 1)*ImageWidth], ImageWidth * sizeof(FColor));
-			Scratch[YOffset + ImageWidth + 1] = FColor(MagentaMask);
+			Scratch[YOffset + ImageWidth + 1] = FColor(BackgroundDWColor);
 		}
 
 		//
-		// Find our initial workset of all rows that have magenta pixels bordering non-magenta pixels.
-		// Also find all rows that have zero magenta pixels - these will never be added to the workset.
+		// Find our initial workset of all rows that have background colored pixels bordering non-background pixels.
+		// Also find all rows that have zero background colored pixels - these will never be added to the workset.
 		//
 
 		TArray<bool> RowCompleted;
@@ -171,28 +163,28 @@ namespace FMaterialBakingHelpersImpl
 		bool bHasAnyData = false;
 		for (int32 Y = 1; Y <= ImageHeight; ++Y)
 		{
-			bool bHasMagenta = false;
-			bool bBordersNonMagenta = false;
+			bool bHasBackground = false;
+			bool bBordersNonBackground = false;
 			const int32 YOffset = Y * PaddedImageWidth;
 			for (int32 X = 1; X <= ImageWidth; ++X)
 			{
-				if (Scratch[YOffset + X].DWColor() == MagentaMask)
+				if (Scratch[YOffset + X].DWColor() == BackgroundDWColor)
 				{
-					bHasMagenta = true;
-					if (HasBorderingPixel(Scratch, X, Y, PaddedImageWidth, PaddedImageHeight))
+					bHasBackground = true;
+					if (HasBorderingPixel(Scratch, X, Y, PaddedImageWidth, PaddedImageHeight, BackgroundDWColor))
 					{
-						bBordersNonMagenta = true;
+						bBordersNonBackground = true;
 						break;
 					}
 				}
 			}
 
-			if (!bHasMagenta)
+			if (!bHasBackground)
 			{
 				bHasAnyData = true;
 				RowCompleted[Y] = true;
 			}
-			else if (bBordersNonMagenta)
+			else if (bBordersNonBackground)
 			{
 				bHasAnyData = true;
 				CurrentRowsLeft->Add(Y);
@@ -227,8 +219,8 @@ namespace FMaterialBakingHelpersImpl
 			const int32 NumThreads = FMath::Min(CurrentRowsLeft->Num(), MaxThreads);
 			const int32 LinesPerThread = FMath::CeilToInt((float)CurrentRowsLeft->Num() / (float)NumThreads);
 
-			// split up rows that still have magenta pixels amongst threads
-			ParallelFor(NumThreads, [ImageWidth, PaddedImageWidth, PaddedImageHeight, Current, Scratch, CurrentRowsLeft, LinesPerThread, &RowRemainingPixels, MagentaMask](int32 ThreadIndex)
+			// split up rows that still have background colored pixels amongst threads
+			ParallelFor(NumThreads, [ImageWidth, PaddedImageWidth, PaddedImageHeight, Current, Scratch, CurrentRowsLeft, LinesPerThread, &RowRemainingPixels, BackgroundDWColor](int32 ThreadIndex)
 			{
 				const int32 StartY = ThreadIndex*LinesPerThread;
 				const int32 EndY = FMath::Min((ThreadIndex + 1) * LinesPerThread, CurrentRowsLeft->Num());
@@ -242,11 +234,11 @@ namespace FMaterialBakingHelpersImpl
 					for (int32 X = 1; X <= ImageWidth; X++)
 					{
 						FColor& Color = Current[PixelIndex++];
-						if (Color.DWColor() == MagentaMask)
+						if (Color.DWColor() == BackgroundDWColor)
 						{
-							const FColor SampledColor = BoxBlurSample(Scratch, X, Y, PaddedImageWidth, PaddedImageHeight);
+							const FColor SampledColor = BoxBlurSample(Scratch, X, Y, PaddedImageWidth, PaddedImageHeight, BackgroundDWColor);
 							// If it's a valid pixel
-							if (SampledColor.DWColor() != MagentaMask)
+							if (SampledColor.DWColor() != BackgroundDWColor)
 							{
 								Color = SampledColor;
 							}
@@ -298,7 +290,7 @@ namespace FMaterialBakingHelpersImpl
 			}
 
 			// Swap rows left buffers
-			{			
+			{
 				TArray<uint32>* Temp = NextRowsLeft;
 				NextRowsLeft = CurrentRowsLeft;
 				CurrentRowsLeft = Temp;
@@ -316,7 +308,7 @@ namespace FMaterialBakingHelpersImpl
 			for (int32 X = 1; X <= ImageWidth; X++)
 			{
 				FColor& Color = Current[PixelIndex++];
-				if (Color.DWColor() == MagentaMask)
+				if (Color.DWColor() == BackgroundDWColor)
 				{
 					Color = FColor::Black;
 				}
@@ -325,20 +317,21 @@ namespace FMaterialBakingHelpersImpl
 	}
 }
 
-void FMaterialBakingHelpers::PerformUVBorderSmear(TArray<FColor>& InOutPixels, int32 ImageWidth, int32 ImageHeight, int32 MaxIterations)
+void FMaterialBakingHelpers::PerformUVBorderSmear(TArray<FColor>& InOutPixels, int32 ImageWidth, int32 ImageHeight, int32 MaxIterations, FColor BackgroundColor)
 {
-	FMaterialBakingHelpersImpl::PerformUVBorderSmear(InOutPixels, ImageWidth, ImageHeight, MaxIterations);
+	FMaterialBakingHelpersImpl::PerformUVBorderSmear(InOutPixels, ImageWidth, ImageHeight, MaxIterations, BackgroundColor.DWColor());
 }
 
-void FMaterialBakingHelpers::PerformShrinking(TArray<FColor>& InOutPixels, int32& InOutImageWidth, int32& InOutImageHeight)
+void FMaterialBakingHelpers::PerformShrinking(TArray<FColor>& InOutPixels, int32& InOutImageWidth, int32& InOutImageHeight, FColor BackgroundColor)
 {
-	FMaterialBakingHelpersImpl::PerformShrinking(InOutPixels, InOutImageWidth, InOutImageHeight);
+	FMaterialBakingHelpersImpl::PerformShrinking(InOutPixels, InOutImageWidth, InOutImageHeight, BackgroundColor.DWColor());
 }
 
-void FMaterialBakingHelpers::PerformUVBorderSmearAndShrink(TArray<FColor>& InOutPixels, int32& InOutImageWidth, int32& InOutImageHeight)
+void FMaterialBakingHelpers::PerformUVBorderSmearAndShrink(TArray<FColor>& InOutPixels, int32& InOutImageWidth, int32& InOutImageHeight, FColor BackgroundColor)
 {
-	FMaterialBakingHelpersImpl::PerformShrinking(InOutPixels, InOutImageWidth, InOutImageHeight);
+	const uint32 BackgroundDWColor = BackgroundColor.DWColor();
+	FMaterialBakingHelpersImpl::PerformShrinking(InOutPixels, InOutImageWidth, InOutImageHeight, BackgroundDWColor);
 
 	const int32 DefaultMaxIteration = -1; // Perform smearing over the whole image
-	FMaterialBakingHelpersImpl::PerformUVBorderSmear(InOutPixels, InOutImageWidth, InOutImageHeight, DefaultMaxIteration);
+	FMaterialBakingHelpersImpl::PerformUVBorderSmear(InOutPixels, InOutImageWidth, InOutImageHeight, DefaultMaxIteration, BackgroundDWColor);
 }
\ No newline at end of file
diff --git a/Engine/Source/Developer/MaterialBaking/Private/MaterialBakingModule.cpp b/Engine/Source/Developer/MaterialBaking/Private/MaterialBakingModule.cpp
--- a/Engine/Source/Developer/MaterialBaking/Private/MaterialBakingModule.cpp
+++ b/Engine/Source/Developer/MaterialBaking/Private/MaterialBakingModule.cpp
@@ -331,6 +331,8 @@ void FMaterialBakingModule::BakeMaterials(const TArray<FMaterialData*>& Material
 		MaterialDataEx.bPerformBorderSmear = MaterialData->bPerformBorderSmear;
 		MaterialDataEx.bPerformShrinking = MaterialData->bPerformShrinking;
 		MaterialDataEx.bTangentSpaceNormal = MaterialData->bTangentSpaceNormal;
+		MaterialDataEx.BlendMode = MaterialData->BlendMode;
+		MaterialDataEx.BackgroundColor = MaterialData->BackgroundColor;
 
 		for (const TPair<EMaterialProperty, FIntPoint>& PropertySizePair : MaterialData->PropertySizes)
 		{
@@ -583,7 +585,7 @@ void FMaterialBakingModule::BakeMaterials(const TArray<FMaterialDataEx*>& Materi
 
 				// It is safe to reuse the same render target for each draw pass since they all execute sequentially on the GPU and are copied to staging buffers before
 				// being reused.
-				UTextureRenderTarget2D* RenderTarget = CreateRenderTarget((ColorSpace == EPropertyColorSpace::Linear), PixelFormat, CurrentOutput.PropertySizes[Property]);
+				UTextureRenderTarget2D* RenderTarget = CreateRenderTarget((ColorSpace == EPropertyColorSpace::Linear), PixelFormat, CurrentOutput.PropertySizes[Property], CurrentMaterialSettings->BackgroundColor);
 				if (RenderTarget != nullptr)
 				{
 					// Perform everything left of the operation directly on the render thread since we need to modify some RenderItem's properties
@@ -659,7 +661,7 @@ void FMaterialBakingModule::BakeMaterials(const TArray<FMaterialDataEx*>& Materi
 									if (Property.Type == MP_EmissiveColor)
 									{
 										// Only one thread will write to CurrentOutput.EmissiveScale since there can be only one emissive channel property per FBakeOutputEx
-										FMaterialBakingModule::ProcessEmissiveOutput((const FFloat16Color*)Data, DataWidth, OutputSize, OutputColor, CurrentOutput.EmissiveScale);
+										FMaterialBakingModule::ProcessEmissiveOutput((const FFloat16Color*)Data, DataWidth, OutputSize, OutputColor, CurrentOutput.EmissiveScale, CurrentMaterialSettings->BackgroundColor);
 
 										if (bEmissiveHDR)
 										{
@@ -681,12 +683,12 @@ void FMaterialBakingModule::BakeMaterials(const TArray<FMaterialDataEx*>& Materi
 
 									if (CurrentMaterialSettings->bPerformShrinking)
 									{
-										FMaterialBakingHelpers::PerformShrinking(OutputColor, OutputSize.X, OutputSize.Y);
+										FMaterialBakingHelpers::PerformShrinking(OutputColor, OutputSize.X, OutputSize.Y, CurrentMaterialSettings->BackgroundColor);
 									}
 
 									if (CurrentMaterialSettings->bPerformBorderSmear)
 									{
-										FMaterialBakingHelpers::PerformUVBorderSmear(OutputColor, OutputSize.X, OutputSize.Y);
+										FMaterialBakingHelpers::PerformUVBorderSmear(OutputColor, OutputSize.X, OutputSize.Y, -1, CurrentMaterialSettings->BackgroundColor);
 									}
 #if WITH_EDITOR
 									// If saving intermediates is turned on
@@ -869,7 +871,7 @@ void FMaterialBakingModule::CleanupMaterialProxies()
 	MaterialProxyPool.Reset();
 }
 
-UTextureRenderTarget2D* FMaterialBakingModule::CreateRenderTarget(bool bInForceLinearGamma, EPixelFormat InPixelFormat, const FIntPoint& InTargetSize)
+UTextureRenderTarget2D* FMaterialBakingModule::CreateRenderTarget(bool bInForceLinearGamma, EPixelFormat InPixelFormat, const FIntPoint& InTargetSize, const FColor& BackgroundColor)
 {
 	TRACE_CPUPROFILER_EVENT_SCOPE(FMaterialBakingModule::CreateRenderTarget)
 
@@ -896,8 +898,7 @@ UTextureRenderTarget2D* FMaterialBakingModule::CreateRenderTarget(bool bInForceL
 		RenderTarget = NewObject<UTextureRenderTarget2D>();
 		check(RenderTarget);
 		RenderTarget->AddToRoot();
-		RenderTarget->ClearColor = FLinearColor(1.0f, 0.0f, 1.0f);
-		RenderTarget->ClearColor.A = 1.0f;
+		RenderTarget->ClearColor = bInForceLinearGamma ? BackgroundColor.ReinterpretAsLinear() : FLinearColor(BackgroundColor);
 		RenderTarget->TargetGamma = 0.0f;
 		RenderTarget->InitCustomFormat(ClampedTargetSize.X, ClampedTargetSize.Y, InPixelFormat, bInForceLinearGamma);
 
@@ -921,7 +922,7 @@ FExportMaterialProxy* FMaterialBakingModule::CreateMaterialProxy(const FMaterial
 	// Look for the matching property
 	for (FMaterialPoolValue& Entry : Entries)
 	{
-		if (Entry.Key == Property && Entry.Value->bTangentSpaceNormal == MaterialSettings->bTangentSpaceNormal)
+		if (Entry.Key == Property && Entry.Value->bTangentSpaceNormal == MaterialSettings->bTangentSpaceNormal && Entry.Value->ProxyBlendMode == MaterialSettings->BlendMode)
 		{
 			Proxy = Entry.Value;
 			break;
@@ -931,14 +932,14 @@ FExportMaterialProxy* FMaterialBakingModule::CreateMaterialProxy(const FMaterial
 	// Not found, create a new entry
 	if (Proxy == nullptr)
 	{
-		Proxy = new FExportMaterialProxy(MaterialSettings->Material, Property.Type, Property.CustomOutput.ToString(), false /* bInSynchronousCompilation */, MaterialSettings->bTangentSpaceNormal);
+		Proxy = new FExportMaterialProxy(MaterialSettings->Material, Property.Type, Property.CustomOutput.ToString(), false /* bInSynchronousCompilation */, MaterialSettings->bTangentSpaceNormal, MaterialSettings->BlendMode);
 		MaterialProxyPool.Add(MaterialSettings->Material, FMaterialPoolValue(Property, Proxy));
 	}
 
 	return Proxy;
 }
 
-void FMaterialBakingModule::ProcessEmissiveOutput(const FFloat16Color* Color16, int32 Color16Pitch, const FIntPoint& OutputSize, TArray<FColor>& OutputColor, float& EmissiveScale)
+void FMaterialBakingModule::ProcessEmissiveOutput(const FFloat16Color* Color16, int32 Color16Pitch, const FIntPoint& OutputSize, TArray<FColor>& OutputColor, float& EmissiveScale, const FColor& BackgroundColor)
 {
 	TRACE_CPUPROFILER_EVENT_SCOPE(FMaterialBakingModule::ProcessEmissiveOutput)
 
@@ -950,13 +951,13 @@ void FMaterialBakingModule::ProcessEmissiveOutput(const FFloat16Color* Color16,
 	float* MaxValue = new float[NumThreads];
 	FMemory::Memset(MaxValue, 0, NumThreads * sizeof(MaxValue[0]));
 	const int32 LinesPerThread = FMath::CeilToInt((float)OutputSize.Y / (float)NumThreads);
+	const FFloat16Color BackgroundColor16 = FFloat16Color(FLinearColor(BackgroundColor)); // Can assume emissive always uses sRGB
 
 	// Find maximum float value across texture
-	ParallelFor(NumThreads, [&Color16, LinesPerThread, MaxValue, OutputSize, Color16Pitch](int32 Index)
+	ParallelFor(NumThreads, [&Color16, LinesPerThread, MaxValue, OutputSize, Color16Pitch, BackgroundColor16](int32 Index)
 	{
 		const int32 EndY = FMath::Min((Index + 1) * LinesPerThread, OutputSize.Y);			
 		float& CurrentMaxValue = MaxValue[Index];
-		const FFloat16Color MagentaFloat16 = FFloat16Color(FLinearColor(1.0f, 0.0f, 1.0f));
 		for (int32 PixelY = Index * LinesPerThread; PixelY < EndY; ++PixelY)
 		{
 			const int32 SrcYOffset = PixelY * Color16Pitch;
@@ -964,7 +965,7 @@ void FMaterialBakingModule::ProcessEmissiveOutput(const FFloat16Color* Color16,
 			{
 				const FFloat16Color& Pixel16 = Color16[PixelX + SrcYOffset];
 				// Find maximum channel value across texture
-				if (!(Pixel16 == MagentaFloat16))
+				if (!(Pixel16 == BackgroundColor16))
 				{
 					CurrentMaxValue = FMath::Max(CurrentMaxValue, FMath::Max3(Pixel16.R.GetFloat(), Pixel16.G.GetFloat(), Pixel16.B.GetFloat()));
 				}
@@ -991,10 +992,8 @@ void FMaterialBakingModule::ProcessEmissiveOutput(const FFloat16Color* Color16,
 	// Now convert Float16 to Color using the scale
 	OutputColor.SetNumUninitialized(OutputSize.X * OutputSize.Y);
 	const float Scale = 255.0f / GlobalMaxValue;
-	ParallelFor(NumThreads, [&Color16, LinesPerThread, &OutputColor, OutputSize, Color16Pitch, Scale](int32 Index)
+	ParallelFor(NumThreads, [&Color16, LinesPerThread, &OutputColor, OutputSize, Color16Pitch, Scale, BackgroundColor16, BackgroundColor](int32 Index)
 	{
-		const FFloat16Color MagentaFloat16 = FFloat16Color(FLinearColor(1.0f, 0.0f, 1.0f));
-
 		const int32 EndY = FMath::Min((Index + 1) * LinesPerThread, OutputSize.Y);
 		for (int32 PixelY = Index * LinesPerThread; PixelY < EndY; ++PixelY)
 		{
@@ -1006,11 +1005,9 @@ void FMaterialBakingModule::ProcessEmissiveOutput(const FFloat16Color* Color16,
 				const FFloat16Color& Pixel16 = Color16[PixelX + SrcYOffset];
 				FColor& Pixel8 = OutputColor[PixelX + DstYOffset];
 
-				if (Pixel16 == MagentaFloat16)
+				if (Pixel16 == BackgroundColor16)
 				{
-					Pixel8.R = 255;
-					Pixel8.G = 0;
-					Pixel8.B = 255;
+					Pixel8 = BackgroundColor;
 				}
 				else
 				{
diff --git a/Engine/Source/Developer/MaterialBaking/Private/MaterialBakingModule.h b/Engine/Source/Developer/MaterialBaking/Private/MaterialBakingModule.h
--- a/Engine/Source/Developer/MaterialBaking/Private/MaterialBakingModule.h
+++ b/Engine/Source/Developer/MaterialBaking/Private/MaterialBakingModule.h
@@ -45,13 +45,13 @@ public:
 
 protected:
 	/* Creates and adds or reuses a RenderTarget from the pool */
-	UTextureRenderTarget2D* CreateRenderTarget(bool bInForceLinearGamma, EPixelFormat InPixelFormat, const FIntPoint& InTargetSize);
+	UTextureRenderTarget2D* CreateRenderTarget(bool bInForceLinearGamma, EPixelFormat InPixelFormat, const FIntPoint& InTargetSize, const FColor& BackgroundColor);
 
 	/* Creates and adds (or reuses a ExportMaterialProxy from the pool if MaterialBaking.UseMaterialProxyCaching is set to 1) */
 	FExportMaterialProxy* CreateMaterialProxy(const FMaterialDataEx* MaterialSettings, const FMaterialPropertyEx& Property);
 
 	/** Helper for emissive color conversion to Output */
-	static void ProcessEmissiveOutput(const FFloat16Color* Color16, int32 Color16Pitch, const FIntPoint& OutputSize, TArray<FColor>& Output, float& EmissiveScale);
+	static void ProcessEmissiveOutput(const FFloat16Color* Color16, int32 Color16Pitch, const FIntPoint& OutputSize, TArray<FColor>& Output, float& EmissiveScale, const FColor& BackgroundColor);
 
 	/** Cleans up all cached material proxies in MaterialProxyPool */
 	void CleanupMaterialProxies();
diff --git a/Engine/Source/Developer/MaterialBaking/Public/MaterialBakingHelpers.h b/Engine/Source/Developer/MaterialBaking/Public/MaterialBakingHelpers.h
--- a/Engine/Source/Developer/MaterialBaking/Public/MaterialBakingHelpers.h
+++ b/Engine/Source/Developer/MaterialBaking/Public/MaterialBakingHelpers.h
@@ -3,7 +3,7 @@
 #pragma once
 
 #include "Containers/Array.h"
-#include "HAL/Platform.h"
+#include "Math/Color.h"
 
 struct FColor;
 
@@ -11,13 +11,14 @@ class MATERIALBAKING_API FMaterialBakingHelpers
 {
 public:
 	/**
-	 * Applies a box blur to magenta pixels found in given texture represented by InOutPixels using non-magenta pixels, this creates a smear across the magenta/filled pixels.
+	 * Applies a box blur to magenta pixels found in given texture represented by InOutPixels using non-background pixels, this creates a smear across the magenta/filled pixels.
 	 * @param  InOutPixels     The image to apply the smear to
 	 * @param  ImageWidth      The width of the image in pixels
 	 * @param  ImageHeight     The height of the image in pixels
 	 * @param  MaxIterations   The max distance in pixels to smear the edges of the texture. When set to -1 this value is the max of the width and height of the buffer
+	 * @param  BackgroundColor The pixel color considered background in the image, and that can be smeared into
 	 */
-	static void PerformUVBorderSmear(TArray<FColor>& InOutPixels, int32 ImageWidth, int32 ImageHeight, int32 MaxIterations = -1);
+	static void PerformUVBorderSmear(TArray<FColor>& InOutPixels, int32 ImageWidth, int32 ImageHeight, int32 MaxIterations = -1, FColor BackgroundColor = FColor::Magenta);
 
 	/**
 	 * Shrinks the array to a single pixel if the entire output is monochrome. Params will be adjusted accordingly when this occurs.
@@ -27,15 +28,16 @@ public:
 	 * @param  ImageWidth      The width of the image in pixels
 	 * @param  ImageHeight     The height of the image in pixels
 	 */
-	static void PerformShrinking(TArray<FColor>& InOutPixels, int32& OnOutImageWidth, int32& InOutImageHeight);
+	static void PerformShrinking(TArray<FColor>& InOutPixels, int32& OnOutImageWidth, int32& InOutImageHeight, FColor BackgroundColor = FColor::Magenta);
 
 	/**
-	 * Applies a box blur to magenta pixels found in given texture represented by InOutPixels using non-magenta pixels, this creates a smear across the magenta/filled pixels.
+	 * Applies a box blur to magenta pixels found in given texture represented by InOutPixels using non-background pixels, this creates a smear across the magenta/filled pixels.
 	 * @param  InOutPixels      The image to apply the smear to
 	 * @param  InOutImageWidth  The width of the image in pixels
 	 * @param  InOutImageHeight The height of the image in pixels
+	 * @param  BackgroundColor  The pixel color considered background in the image, and that can be smeared into
 	 *
 	 * @note   This version is allowed to shrink the array to a single pixel if the entire output is monochrome. Params will be adjusted accordingly when this occurs.
 	 */
-	static void PerformUVBorderSmearAndShrink(TArray<FColor>& InOutPixels, int32& InOutImageWidth, int32& InOutImageHeight);
-};
\ No newline at end of file
+	static void PerformUVBorderSmearAndShrink(TArray<FColor>& InOutPixels, int32& InOutImageWidth, int32& InOutImageHeight, FColor BackgroundColor = FColor::Magenta);
+};
diff --git a/Engine/Source/Developer/MaterialBaking/Public/MaterialBakingStructures.h b/Engine/Source/Developer/MaterialBaking/Public/MaterialBakingStructures.h
--- a/Engine/Source/Developer/MaterialBaking/Public/MaterialBakingStructures.h
+++ b/Engine/Source/Developer/MaterialBaking/Public/MaterialBakingStructures.h
@@ -31,6 +31,12 @@ struct FMaterialData
 
 	/** Whether to transform normals from world-space to tangent-space (does nothing if material already uses tangent-space normals) */
 	bool bTangentSpaceNormal = false;
+
+	/** Blend mode to use when baking, allowing for example detection of overlapping UVs */
+	EBlendMode BlendMode = BLEND_Opaque;
+
+	/** Background color used to initially fill the output texture and used for border smear */
+	FColor BackgroundColor = FColor::Magenta;
 };
 
 /** Structure containing extended information about the material and properties which is being baked out */
@@ -50,6 +56,12 @@ struct FMaterialDataEx
 
 	/** Whether to transform normals from world-space to tangent-space (does nothing if material already uses tangent-space normals) */
 	bool bTangentSpaceNormal = false;
+
+	/** Blend mode to use when baking, allowing for example detection of overlapping UVs */
+	EBlendMode BlendMode = BLEND_Opaque;
+
+	/** Background color used to initially fill the output texture and used for border smear */
+	FColor BackgroundColor = FColor::Magenta;
 };
 
 /** Structure containing primitive information (regarding a mesh or mesh component) that is accessible through material expressions */
