Add support for material baking refraction property

(based on PR #8496)

diff --git a/Engine/Source/Developer/MaterialBaking/Private/ExportMaterialProxy.h b/Engine/Source/Developer/MaterialBaking/Private/ExportMaterialProxy.h
--- a/Engine/Source/Developer/MaterialBaking/Private/ExportMaterialProxy.h
+++ b/Engine/Source/Developer/MaterialBaking/Private/ExportMaterialProxy.h
@@ -249,6 +249,7 @@ public:
 		case MP_AmbientOcclusion: ResourceId.Usage = EMaterialShaderMapUsage::MaterialExportAO; break;
 		case MP_EmissiveColor: ResourceId.Usage = EMaterialShaderMapUsage::MaterialExportEmissive; break;
 		case MP_Opacity: ResourceId.Usage = EMaterialShaderMapUsage::MaterialExportOpacity; break;
+		case MP_Refraction: ResourceId.Usage = EMaterialShaderMapUsage::MaterialExportRefraction; break;
 		case MP_OpacityMask: ResourceId.Usage = EMaterialShaderMapUsage::MaterialExportOpacityMask; break;
 		case MP_SubsurfaceColor: ResourceId.Usage = EMaterialShaderMapUsage::MaterialExportSubSurfaceColor; break;
 		case MP_ShadingModel: ResourceId.Usage = EMaterialShaderMapUsage::MaterialExportShadingModel; break;
@@ -380,6 +381,15 @@ public:
 						CompileNormalTransform(&ProxyCompiler, MaterialInterface->CompileProperty(&ProxyCompiler, PropertyToCompile, ForceCast_Exact_Replicate)));
 				}
 				break;
+			case MP_Refraction:
+				// Only index of refraction can be supported because other methods don't have values within a suitable range for encoding into 8-bit baked textures
+				if (Material->RefractionMethod == RM_IndexOfRefraction)
+				{
+					return CompileRefractionEncoding(
+						Compiler,
+						MaterialInterface->CompileProperty(&ProxyCompiler, MP_Refraction, ForceCast_Exact_Replicate));
+				}
+				break;
 			case MP_ShadingModel:
 				return CompileShadingModelEncoding(Compiler, MaterialInterface->CompileProperty(&ProxyCompiler, MP_ShadingModel));
 			case MP_CustomOutput:
@@ -604,6 +614,14 @@ private:
 			Compiler->Constant(0.5f)); // [-0.5,0.5] + 0.5
 	}
 
+	static int32 CompileRefractionEncoding(FMaterialCompiler* Compiler, int32 RefractionInput)
+	{
+		// [1,Infinity] -> [1,0]
+		return Compiler->Div(
+			Compiler->Constant(1.0f),
+			Compiler->Max(Compiler->Constant(1.0f), RefractionInput));
+	}
+
 	static int32 CompileShadingModelEncoding(FMaterialCompiler* Compiler, int32 ShadingModelInput)
 	{
 		// [0,MSM_NUM] -> [0,1]
diff --git a/Engine/Source/Developer/MaterialBaking/Private/MaterialBakingModule.cpp b/Engine/Source/Developer/MaterialBaking/Private/MaterialBakingModule.cpp
--- a/Engine/Source/Developer/MaterialBaking/Private/MaterialBakingModule.cpp
+++ b/Engine/Source/Developer/MaterialBaking/Private/MaterialBakingModule.cpp
@@ -283,6 +283,7 @@ void FMaterialBakingModule::StartupModule()
 	PerPropertyFormat.Add(MP_Specular, PF_B8G8R8A8);
 	PerPropertyFormat.Add(MP_Roughness, PF_B8G8R8A8);
 	PerPropertyFormat.Add(MP_Anisotropy, PF_B8G8R8A8);
+	PerPropertyFormat.Add(MP_Refraction, PF_B8G8R8A8);
 	PerPropertyFormat.Add(MP_Normal, PF_B8G8R8A8);
 	PerPropertyFormat.Add(MP_Tangent, PF_B8G8R8A8);
 	PerPropertyFormat.Add(MP_AmbientOcclusion, PF_B8G8R8A8);
@@ -846,6 +847,7 @@ void FMaterialBakingModule::SetLinearBake(bool bCorrectLinear)
 	{
 		DefaultColorSpace = EPropertyColorSpace::sRGB;
 		PerPropertyColorSpace.Add(MP_Normal, EPropertyColorSpace::Linear);
+		PerPropertyColorSpace.Add(MP_Refraction, EPropertyColorSpace::Linear);
 		PerPropertyColorSpace.Add(MP_Opacity, EPropertyColorSpace::Linear);
 		PerPropertyColorSpace.Add(MP_OpacityMask, EPropertyColorSpace::Linear);
 		PerPropertyColorSpace.Add(MP_ShadingModel, EPropertyColorSpace::Linear);
diff --git a/Engine/Source/Runtime/Engine/Public/MaterialShared.h b/Engine/Source/Runtime/Engine/Public/MaterialShared.h
--- a/Engine/Source/Runtime/Engine/Public/MaterialShared.h
+++ b/Engine/Source/Runtime/Engine/Public/MaterialShared.h
@@ -797,6 +797,7 @@ namespace EMaterialShaderMapUsage
 		MaterialExportClearCoatRoughness,
 		MaterialExportCustomOutput,
 		MaterialExportShadingModel,
+		MaterialExportRefraction,
 	};
 }
 
