Add support for material baking custom output transmittance color and fix crash when baking invalid custom output

(based on PR #8498)

diff --git a/Engine/Source/Developer/MaterialBaking/Private/ExportMaterialProxy.h b/Engine/Source/Developer/MaterialBaking/Private/ExportMaterialProxy.h
--- a/Engine/Source/Developer/MaterialBaking/Private/ExportMaterialProxy.h
+++ b/Engine/Source/Developer/MaterialBaking/Private/ExportMaterialProxy.h
@@ -393,8 +393,12 @@ public:
 			case MP_ShadingModel:
 				return CompileShadingModelEncoding(Compiler, MaterialInterface->CompileProperty(&ProxyCompiler, MP_ShadingModel));
 			case MP_CustomOutput:
-				 // NOTE: Currently we can assume input index is always 0, which it is for all custom outputs that are registered as material attributes
-				return CompileInputForCustomOutput(&ProxyCompiler, 0, ForceCast_Exact_Replicate);
+				if (const FMaterialCustomOutputAttributeDefintion* CustomAttribute = FMaterialAttributeDefinitionMap::GetCustomAttribute(CustomOutputToCompile))
+				{
+					constexpr int32 InputIndex = 0; // Assume input index is always 0, which it is for all custom outputs that are registered as material attributes
+					return CompileInputForCustomOutput(&ProxyCompiler, CustomAttribute, InputIndex, ForceCast_Exact_Replicate);
+				}
+				break;
 			default:
 				return Compiler->Constant(1.0f);
 			}
@@ -556,14 +560,13 @@ public:
 	}
 
 private:
-	int32 CompileInputForCustomOutput(FMaterialCompiler* Compiler, int32 InputIndex, uint32 ForceCastFlags) const
+	int32 CompileInputForCustomOutput(FMaterialCompiler* Compiler, const FMaterialCustomOutputAttributeDefintion* CustomAttribute, int32 InputIndex, uint32 ForceCastFlags) const
 	{
-		FGuid AttributeID = FMaterialAttributeDefinitionMap::GetCustomAttributeID(CustomOutputToCompile);
-		check(AttributeID.IsValid());
+		check(CustomAttribute);
 
-		UMaterialExpressionCustomOutput* Expression = GetCustomOutputExpressionToCompile();
+		UMaterialExpressionCustomOutput* Expression = GetCustomOutputExpression(CustomAttribute->FunctionName);
 		FExpressionInput* ExpressionInput = Expression ? Expression->GetInput(InputIndex) : nullptr;
-		int32 Result = INDEX_NONE;
+		int32 Result;
 
 		if (ExpressionInput)
 		{
@@ -571,7 +574,7 @@ private:
 		}
 		else
 		{
-			Result = FMaterialAttributeDefinitionMap::CompileDefaultExpression(Compiler, AttributeID);
+			Result = CustomAttribute->CompileDefaultValue(Compiler);
 		}
 
 		if (CustomOutputToCompile == TEXT("ClearCoatBottomNormal"))
@@ -581,18 +584,18 @@ private:
 
 		if (ForceCastFlags & MFCF_ForceCast)
 		{
-			Result = Compiler->ForceCast(Result, FMaterialAttributeDefinitionMap::GetValueType(AttributeID), ForceCastFlags);
+			Result = Compiler->ForceCast(Result, CustomAttribute->ValueType, ForceCastFlags);
 		}
 
 		return Result;
 	}
 
-	UMaterialExpressionCustomOutput* GetCustomOutputExpressionToCompile() const
+	UMaterialExpressionCustomOutput* GetCustomOutputExpression(const FString& FunctionName) const
 	{
 		for (UMaterialExpression* Expression : Material->GetExpressions())
 		{
 			UMaterialExpressionCustomOutput* CustomOutputExpression = Cast<UMaterialExpressionCustomOutput>(Expression);
-			if (CustomOutputExpression && CustomOutputExpression->GetDisplayName() == CustomOutputToCompile)
+			if (CustomOutputExpression && CustomOutputExpression->GetFunctionName() == FunctionName)
 			{
 				return CustomOutputExpression;
 			}
diff --git a/Engine/Source/Developer/MaterialBaking/Private/MaterialBakingModule.cpp b/Engine/Source/Developer/MaterialBaking/Private/MaterialBakingModule.cpp
--- a/Engine/Source/Developer/MaterialBaking/Private/MaterialBakingModule.cpp
+++ b/Engine/Source/Developer/MaterialBaking/Private/MaterialBakingModule.cpp
@@ -291,7 +291,8 @@ void FMaterialBakingModule::StartupModule()
 	PerPropertyFormat.Add(MP_CustomData0, PF_B8G8R8A8);
 	PerPropertyFormat.Add(MP_CustomData1, PF_B8G8R8A8);
 	PerPropertyFormat.Add(MP_ShadingModel, PF_B8G8R8A8);
-	PerPropertyFormat.Add(TEXT("ClearCoatBottomNormal"), PF_B8G8R8A8);
+	PerPropertyFormat.Add(FMaterialPropertyEx::ClearCoatBottomNormal, PF_B8G8R8A8);
+	PerPropertyFormat.Add(FMaterialPropertyEx::TransmittanceColor, PF_B8G8R8A8);
 
 	// Register property customization
 	FPropertyEditorModule& Module = FModuleManager::Get().LoadModuleChecked<FPropertyEditorModule>("PropertyEditor");
@@ -842,6 +843,7 @@ void FMaterialBakingModule::SetLinearBake(bool bCorrectLinear)
 		PerPropertyColorSpace.Add(MP_BaseColor, EPropertyColorSpace::sRGB);
 		PerPropertyColorSpace.Add(MP_EmissiveColor, EPropertyColorSpace::sRGB);
 		PerPropertyColorSpace.Add(MP_SubsurfaceColor, EPropertyColorSpace::sRGB);
+		PerPropertyColorSpace.Add(FMaterialPropertyEx::TransmittanceColor, EPropertyColorSpace::sRGB);
 	}
 	else
 	{
@@ -851,7 +853,7 @@ void FMaterialBakingModule::SetLinearBake(bool bCorrectLinear)
 		PerPropertyColorSpace.Add(MP_Opacity, EPropertyColorSpace::Linear);
 		PerPropertyColorSpace.Add(MP_OpacityMask, EPropertyColorSpace::Linear);
 		PerPropertyColorSpace.Add(MP_ShadingModel, EPropertyColorSpace::Linear);
-		PerPropertyColorSpace.Add(TEXT("ClearCoatBottomNormal"), EPropertyColorSpace::Linear);
+		PerPropertyColorSpace.Add(FMaterialPropertyEx::ClearCoatBottomNormal, EPropertyColorSpace::Linear);
 	}
 }
 
diff --git a/Engine/Source/Developer/MaterialBaking/Private/MaterialPropertyEx.cpp b/Engine/Source/Developer/MaterialBaking/Private/MaterialPropertyEx.cpp
--- a/Engine/Source/Developer/MaterialBaking/Private/MaterialPropertyEx.cpp
+++ b/Engine/Source/Developer/MaterialBaking/Private/MaterialPropertyEx.cpp
@@ -3,6 +3,9 @@
 #include "MaterialPropertyEx.h"
 #include "UObject/NameTypes.h"
 
+const FMaterialPropertyEx FMaterialPropertyEx::ClearCoatBottomNormal(TEXT("ClearCoatBottomNormal"));
+const FMaterialPropertyEx FMaterialPropertyEx::TransmittanceColor(TEXT("TransmittanceColor"));
+
 FString FMaterialPropertyEx::ToString() const
 {
 	if (!IsCustomOutput())
diff --git a/Engine/Source/Developer/MaterialBaking/Public/MaterialPropertyEx.h b/Engine/Source/Developer/MaterialBaking/Public/MaterialPropertyEx.h
--- a/Engine/Source/Developer/MaterialBaking/Public/MaterialPropertyEx.h
+++ b/Engine/Source/Developer/MaterialBaking/Public/MaterialPropertyEx.h
@@ -10,7 +10,7 @@
 #include "UObject/UnrealNames.h"
 
 /** Structure extending EMaterialProperty to allow detailed information about custom output */
-struct FMaterialPropertyEx
+struct MATERIALBAKING_API FMaterialPropertyEx
 {
 	FMaterialPropertyEx(EMaterialProperty Type = MP_MAX, const FName& CustomOutput = NAME_None)
 		: Type(Type)
@@ -47,11 +47,17 @@ struct FMaterialPropertyEx
 		return !Other.IsCustomOutput() ? GetTypeHash(Other.Type) : GetTypeHash(Other.CustomOutput);
 	}
 
-	MATERIALBAKING_API FString ToString() const;
+	FString ToString() const;
 
 	/** The material property */
 	EMaterialProperty Type;
 
 	/** The name of a specific custom output. Only used if property is MP_CustomOutput */
 	FName CustomOutput;
+
+	/** The identifier for custom output Clear Coat Bottom Normal (used by shading model Clear Coat) */
+	static const FMaterialPropertyEx ClearCoatBottomNormal;
+
+	/** The identifier for custom output Transmittance Color (used by shading model Thin Translucent) */
+	static const FMaterialPropertyEx TransmittanceColor;
 };
diff --git a/Engine/Source/Runtime/Engine/Private/Materials/MaterialAttributeDefinitionMap.cpp b/Engine/Source/Runtime/Engine/Private/Materials/MaterialAttributeDefinitionMap.cpp
--- a/Engine/Source/Runtime/Engine/Private/Materials/MaterialAttributeDefinitionMap.cpp
+++ b/Engine/Source/Runtime/Engine/Private/Materials/MaterialAttributeDefinitionMap.cpp
@@ -28,7 +28,7 @@ FMaterialAttributeDefintion::FMaterialAttributeDefintion(
 	checkf(ValueType & MCT_Float || ValueType == MCT_ShadingModel || ValueType == MCT_Strata || ValueType == MCT_MaterialAttributes, TEXT("Unsupported material attribute type %d"), ValueType);
 }
 
-int32 FMaterialAttributeDefintion::CompileDefaultValue(FMaterialCompiler* Compiler)
+int32 FMaterialAttributeDefintion::CompileDefaultValue(FMaterialCompiler* Compiler) const
 {
 	int32 Ret;
 
@@ -299,6 +299,7 @@ void FMaterialAttributeDefinitionMap::InitializeAttributeMap()
 	AddCustomAttribute(FGuid(0xfbd7b46e, 0xb1234824, 0xbde76b23, 0x609f984c), "BentNormal", "GetBentNormal", MCT_Float3, FVector4(0, 0, 1, 0));
 	AddCustomAttribute(FGuid(0xAA3D5C04, 0x16294716, 0xBBDEC869, 0x6A27DD72), "ClearCoatBottomNormal", "ClearCoatBottomNormal", MCT_Float3, FVector4(0, 0, 1, 0));
 	AddCustomAttribute(FGuid(0x8EAB2CB2, 0x73634A24, 0x8CD14F47, 0x3F9C8E55), "CustomEyeTangent", "GetTangentOutput", MCT_Float3, FVector4(0, 0, 0, 0));
+	AddCustomAttribute(FGuid(0xF2D8C70E, 0x42ECA0D1, 0x4652D0AD, 0xB785A065), "TransmittanceColor", "GetThinTranslucentMaterialOutput", MCT_Float3, FVector4(0.5, 0.5, 0.5, 0));
 }
 
 void FMaterialAttributeDefinitionMap::Add(const FGuid& AttributeID, const FString& AttributeName, EMaterialProperty Property,
@@ -542,6 +543,19 @@ FGuid FMaterialAttributeDefinitionMap::GetCustomAttributeID(const FString& Attri
 	return GMaterialPropertyAttributesMap.Find(MP_MAX)->AttributeID;
 }
 
+const FMaterialCustomOutputAttributeDefintion* FMaterialAttributeDefinitionMap::GetCustomAttribute(const FString& AttributeName)
+{
+	for (auto& Attribute : GMaterialPropertyAttributesMap.CustomAttributes)
+	{
+		if (Attribute.AttributeName == AttributeName)
+		{
+			return &Attribute;
+		}
+	}
+
+	return nullptr;
+}
+
 void FMaterialAttributeDefinitionMap::GetCustomAttributeList(TArray<FMaterialCustomOutputAttributeDefintion>& CustomAttributeList)
 {
 	CustomAttributeList.Empty(GMaterialPropertyAttributesMap.CustomAttributes.Num());
diff --git a/Engine/Source/Runtime/Engine/Public/Materials/MaterialAttributeDefinitionMap.h b/Engine/Source/Runtime/Engine/Public/Materials/MaterialAttributeDefinitionMap.h
--- a/Engine/Source/Runtime/Engine/Public/Materials/MaterialAttributeDefinitionMap.h
+++ b/Engine/Source/Runtime/Engine/Public/Materials/MaterialAttributeDefinitionMap.h
@@ -32,7 +32,7 @@ public:
 		EMaterialValueType InValueType, const FVector4& InDefaultValue, EShaderFrequency InShaderFrequency,
 		int32 InTexCoordIndex = INDEX_NONE, bool bInIsHidden = false, MaterialAttributeBlendFunction InBlendFunction = nullptr);
 
-	ENGINE_API int32 CompileDefaultValue(FMaterialCompiler* Compiler);
+	ENGINE_API int32 CompileDefaultValue(FMaterialCompiler* Compiler) const;
 
 	bool operator==(const FMaterialAttributeDefintion& Other) const
 	{
@@ -137,6 +137,9 @@ public:
 	/** Returns the first custom attribute ID that has the specificed attribute name */
 	static FGuid GetCustomAttributeID(const FString& AttributeName);
 
+	/** Returns the first custom attribute definition that has the specificed attribute name */
+	static const FMaterialCustomOutputAttributeDefintion* GetCustomAttribute(const FString& AttributeName);
+
 	/** Returns a list of registered custom attributes */
 	static void GetCustomAttributeList(TArray<FMaterialCustomOutputAttributeDefintion>& CustomAttributeList);
 
