Fix incorrect results when material baking some properties in a translucent material. To also avoid future issues like this, the fix removes all hard-coded assumptions in the baking module, regarding which properties are active depending on blend mode. These assumptions are not feasible to manually keep in-sync with the latest material behavior and therefore it's better to let the caller of the material baking module decide how to handle potentially inactive properties.

diff --git a/Engine/Source/Developer/MaterialBaking/Private/ExportMaterialProxy.h b/Engine/Source/Developer/MaterialBaking/Private/ExportMaterialProxy.h
--- a/Engine/Source/Developer/MaterialBaking/Private/ExportMaterialProxy.h
+++ b/Engine/Source/Developer/MaterialBaking/Private/ExportMaterialProxy.h
@@ -346,41 +346,23 @@ public:
 			switch (PropertyToCompile)
 			{
 			case MP_EmissiveColor:
-				// Emissive is ALWAYS returned...
-				return MaterialInterface->CompileProperty(&ProxyCompiler, MP_EmissiveColor, ForceCast_Exact_Replicate);
 			case MP_BaseColor:
-				return MaterialInterface->CompileProperty(&ProxyCompiler, MP_BaseColor, ForceCast_Exact_Replicate);
-				break;
 			case MP_Specular:
 			case MP_Roughness:
 			case MP_Anisotropy:
 			case MP_Metallic:
 			case MP_AmbientOcclusion:
-				// Only return for Opaque and Masked...
-				if (IsOpaqueOrMaskedBlendMode(*MaterialInterface))
-				{
-					return MaterialInterface->CompileProperty(&ProxyCompiler, PropertyToCompile, ForceCast_Exact_Replicate);
-				}
-				break;
-
 			case MP_Opacity:
 			case MP_OpacityMask:
 			case MP_CustomData0:
 			case MP_CustomData1:
 			case MP_SubsurfaceColor:
-			{
 				return MaterialInterface->CompileProperty(&ProxyCompiler, PropertyToCompile, ForceCast_Exact_Replicate);
-			}
 			case MP_Normal:
 			case MP_Tangent:
-				// Only return for Opaque and Masked...
-				if (IsOpaqueOrMaskedBlendMode(*MaterialInterface))
-				{
-					return CompileNormalEncoding(
-						Compiler,
-						CompileNormalTransform(&ProxyCompiler, MaterialInterface->CompileProperty(&ProxyCompiler, PropertyToCompile, ForceCast_Exact_Replicate)));
-				}
-				break;
+				return CompileNormalEncoding(
+					Compiler,
+					CompileNormalTransform(&ProxyCompiler, MaterialInterface->CompileProperty(&ProxyCompiler, PropertyToCompile, ForceCast_Exact_Replicate)));
 			case MP_Refraction:
 				// Only index of refraction can be supported because other methods don't have values within a suitable range for encoding into 8-bit baked textures
 				if (Material->RefractionMethod == RM_IndexOfRefraction)
