Fix incorrect normals when baking certain materials. The issue is due to a change in 5.2, which introduces a new virtual method, FMaterialResource::IsTangentSpaceNormal, that determines whether a material uses tangent-space normals or world-space. The baking module needs to override this method and forward all IsTangentSpaceNormal calls to the actual material being baked to properly determine which space the normals should be in.

diff --git a/Engine/Source/Developer/MaterialBaking/Private/ExportMaterialProxy.h b/Engine/Source/Developer/MaterialBaking/Private/ExportMaterialProxy.h
--- a/Engine/Source/Developer/MaterialBaking/Private/ExportMaterialProxy.h
+++ b/Engine/Source/Developer/MaterialBaking/Private/ExportMaterialProxy.h
@@ -432,6 +432,14 @@ public:
 		// it needs to be a surface material.
 		return MD_Surface;
 	}
+	virtual bool IsTangentSpaceNormal() const override
+	{
+		if (const FMaterialResource* Resource = MaterialInterface->GetMaterialResource(GMaxRHIFeatureLevel))
+		{
+			return Resource->IsTangentSpaceNormal();
+		}
+		return false;
+	}
 	virtual bool IsTwoSided() const  override
 	{
 		if (MaterialInterface)
