Material Baking - Added "MaterialBaking.ForceDisableEmissiveScaling" CVar (disabled by default)
* If set to true, values stored in the emissive textures will be clamped to the [0, 1] range rather than being normalized and scaled back using the EmissiveScale material static parameter

diff --git a/Engine/Source/Developer/MaterialBaking/Private/MaterialBakingModule.cpp b/Engine/Source/Developer/MaterialBaking/Private/MaterialBakingModule.cpp
--- a/Engine/Source/Developer/MaterialBaking/Private/MaterialBakingModule.cpp
+++ b/Engine/Source/Developer/MaterialBaking/Private/MaterialBakingModule.cpp
@@ -66,6 +66,10 @@ static TAutoConsoleVariable<int32> CVarMaterialBakingVTWarmupFrames(
 	5,
 	TEXT("Number of frames to render for virtual texture warmup when material baking."));
 
+static TAutoConsoleVariable<bool> CVarMaterialBakingForceDisableEmissiveScaling(
+	TEXT("MaterialBaking.ForceDisableEmissiveScaling"),
+	false,
+	TEXT("If set to true, values stored in the emissive textures will be clamped to the [0, 1] range rather than being normalized and scaled back using the EmissiveScale material static parameter."));
 
 namespace FMaterialBakingModuleImpl
 {
@@ -958,46 +962,52 @@ void FMaterialBakingModule::ProcessEmissiveOutput(const FFloat16Color* Color16,
 	{
 		return FPlatformProcess::SupportsMultithreading() ? FPlatformMisc::NumberOfCores() : 1;
 	}();
-
-	float* MaxValue = new float[NumThreads];
-	FMemory::Memset(MaxValue, 0, NumThreads * sizeof(MaxValue[0]));
 	const int32 LinesPerThread = FMath::CeilToInt((float)OutputSize.Y / (float)NumThreads);
 	const FFloat16Color BackgroundColor16 = FFloat16Color(FLinearColor(BackgroundColor)); // Can assume emissive always uses sRGB
+	const bool bShouldNormalize = CVarMaterialBakingForceDisableEmissiveScaling.GetValueOnAnyThread() == 0;
+	float GlobalMaxValue = 1.0f;
 
-	// Find maximum float value across texture
-	ParallelFor(NumThreads, [&Color16, LinesPerThread, MaxValue, OutputSize, Color16Pitch, BackgroundColor16](int32 Index)
+	if (bShouldNormalize)
 	{
-		const int32 EndY = FMath::Min((Index + 1) * LinesPerThread, OutputSize.Y);			
-		float& CurrentMaxValue = MaxValue[Index];
-		for (int32 PixelY = Index * LinesPerThread; PixelY < EndY; ++PixelY)
+		float* MaxValue = new float[NumThreads];
+		FMemory::Memset(MaxValue, 0, NumThreads * sizeof(MaxValue[0]));
+		
+		// Find maximum float value across texture
+		ParallelFor(NumThreads, [&Color16, LinesPerThread, MaxValue, OutputSize, Color16Pitch, BackgroundColor16](int32 Index)
 		{
-			const int32 SrcYOffset = PixelY * Color16Pitch;
-			for (int32 PixelX = 0; PixelX < OutputSize.X; PixelX++)
+			const int32 EndY = FMath::Min((Index + 1) * LinesPerThread, OutputSize.Y);			
+			float& CurrentMaxValue = MaxValue[Index];
+			for (int32 PixelY = Index * LinesPerThread; PixelY < EndY; ++PixelY)
 			{
-				const FFloat16Color& Pixel16 = Color16[PixelX + SrcYOffset];
-				// Find maximum channel value across texture
-				if (!(Pixel16 == BackgroundColor16))
+				const int32 SrcYOffset = PixelY * Color16Pitch;
+				for (int32 PixelX = 0; PixelX < OutputSize.X; PixelX++)
 				{
-					CurrentMaxValue = FMath::Max(CurrentMaxValue, FMath::Max3(Pixel16.R.GetFloat(), Pixel16.G.GetFloat(), Pixel16.B.GetFloat()));
+					const FFloat16Color& Pixel16 = Color16[PixelX + SrcYOffset];
+					// Find maximum channel value across texture
+					if (!(Pixel16 == BackgroundColor16))
+					{
+						CurrentMaxValue = FMath::Max(CurrentMaxValue, FMath::Max3(Pixel16.R.GetFloat(), Pixel16.G.GetFloat(), Pixel16.B.GetFloat()));
+					}
 				}
 			}
-		}
-	});
+		});
 
-	const float GlobalMaxValue = [&MaxValue, NumThreads]
-	{
-		float TempValue = 0.0f;
-		for (int32 ThreadIndex = 0; ThreadIndex < NumThreads; ++ThreadIndex)
+		GlobalMaxValue = [&MaxValue, NumThreads]
 		{
-			TempValue = FMath::Max(TempValue, MaxValue[ThreadIndex]);
-		}
+			float TempValue = 0.0f;
+			for (int32 ThreadIndex = 0; ThreadIndex < NumThreads; ++ThreadIndex)
+			{
+				TempValue = FMath::Max(TempValue, MaxValue[ThreadIndex]);
+			}
 
-		return TempValue;
-	}();
+			return TempValue;
+		}();
 		
-	if (GlobalMaxValue <= 0.01f)
-	{
-		// Black emissive, drop it			
+		if (GlobalMaxValue <= 0.01f)
+		{
+			// Black emissive, no need to scale
+			GlobalMaxValue = 1.0f;
+		}
 	}
 
 	// Now convert Float16 to Color using the scale
@@ -1022,9 +1032,9 @@ void FMaterialBakingModule::ProcessEmissiveOutput(const FFloat16Color* Color16,
 				}
 				else
 				{
-					Pixel8.R = (uint8)FMath::RoundToInt(Pixel16.R.GetFloat() * Scale);
-					Pixel8.G = (uint8)FMath::RoundToInt(Pixel16.G.GetFloat() * Scale);
-					Pixel8.B = (uint8)FMath::RoundToInt(Pixel16.B.GetFloat() * Scale);
+					Pixel8.R = (uint8)FMath::Clamp(FMath::RoundToInt(Pixel16.R.GetFloat() * Scale), 0, 255);
+					Pixel8.G = (uint8)FMath::Clamp(FMath::RoundToInt(Pixel16.G.GetFloat() * Scale), 0, 255);
+					Pixel8.B = (uint8)FMath::Clamp(FMath::RoundToInt(Pixel16.B.GetFloat() * Scale), 0, 255);
 				}
 					
 				Pixel8.A = 255;
