Added a CRC to the MaterialBaking module, to be able to detect if code (by changing the base key) or setting changes could require an HLOD rebuild

diff --git a/Engine/Plugins/Editor/WorldPartitionHLODUtilities/Source/Private/HLODBuilderMeshMerge.cpp b/Engine/Plugins/Editor/WorldPartitionHLODUtilities/Source/Private/HLODBuilderMeshMerge.cpp
--- a/Engine/Plugins/Editor/WorldPartitionHLODUtilities/Source/Private/HLODBuilderMeshMerge.cpp
+++ b/Engine/Plugins/Editor/WorldPartitionHLODUtilities/Source/Private/HLODBuilderMeshMerge.cpp
@@ -8,7 +8,7 @@
 #include "Materials/Material.h"
 #include "Components/StaticMeshComponent.h"
 
-
+#include "IMaterialBakingModule.h"
 #include "IMeshMergeUtilities.h"
 #include "MeshMergeModule.h"
 #include "Modules/ModuleManager.h"
@@ -49,6 +49,13 @@ uint32 UHLODBuilderMeshMergeSettings::GetCRC() const
 	Ar << This.MeshMergeSettings;
 	UE_LOG(LogHLODBuilder, VeryVerbose, TEXT(" - MeshMergeSettings = %d"), Ar.GetCrc());
 
+	if (MeshMergeSettings.bMergeMaterials)
+	{
+		IMaterialBakingModule& Module = FModuleManager::Get().LoadModuleChecked<IMaterialBakingModule>("MaterialBaking");
+		uint32 MaterialBakingModuleCRC = Module.GetCRC();
+		Ar << MaterialBakingModuleCRC;
+	}
+
 	uint32 Hash = Ar.GetCrc();
 
 	if (!HLODMaterial.IsNull())
diff --git a/Engine/Plugins/Editor/WorldPartitionHLODUtilities/Source/Private/HLODBuilderMeshSimplify.cpp b/Engine/Plugins/Editor/WorldPartitionHLODUtilities/Source/Private/HLODBuilderMeshSimplify.cpp
--- a/Engine/Plugins/Editor/WorldPartitionHLODUtilities/Source/Private/HLODBuilderMeshSimplify.cpp
+++ b/Engine/Plugins/Editor/WorldPartitionHLODUtilities/Source/Private/HLODBuilderMeshSimplify.cpp
@@ -8,7 +8,7 @@
 #include "Materials/Material.h"
 #include "Components/StaticMeshComponent.h"
 
-
+#include "IMaterialBakingModule.h"
 #include "IMeshMergeUtilities.h"
 #include "MeshMergeModule.h"
 #include "Modules/ModuleManager.h"
@@ -43,9 +43,17 @@ uint32 UHLODBuilderMeshSimplifySettings::GetCRC() const
 
 	FArchiveCrc32 Ar;
 
+	// Base mesh simplify key, changing this will force a rebuild of all HLODs from this builder
+	FString HLODBaseKey = "9015D801CD0E420B9D2B0CB997A97813";
+	Ar << HLODBaseKey;
+
 	Ar << This.MeshSimplifySettings;
 	UE_LOG(LogHLODBuilder, VeryVerbose, TEXT(" - MeshSimplifySettings = %d"), Ar.GetCrc());
 
+	IMaterialBakingModule& Module = FModuleManager::Get().LoadModuleChecked<IMaterialBakingModule>("MaterialBaking");
+	uint32 MaterialBakingModuleCRC = Module.GetCRC();
+	Ar << MaterialBakingModuleCRC;
+
 	uint32 Hash = Ar.GetCrc();
 
 	if (!HLODMaterial.IsNull())
diff --git a/Engine/Plugins/Editor/WorldPartitionHLODUtilities/Source/WorldPartitionHLODUtilities.Build.cs b/Engine/Plugins/Editor/WorldPartitionHLODUtilities/Source/WorldPartitionHLODUtilities.Build.cs
--- a/Engine/Plugins/Editor/WorldPartitionHLODUtilities/Source/WorldPartitionHLODUtilities.Build.cs
+++ b/Engine/Plugins/Editor/WorldPartitionHLODUtilities/Source/WorldPartitionHLODUtilities.Build.cs
@@ -17,6 +17,7 @@ public class WorldPartitionHLODUtilities : ModuleRules
             new string[]
 			{
 				"Engine",
+				"MaterialBaking",
 				"MaterialUtilities",
 				"MeshDescription",
 				"MeshMergeUtilities",
diff --git a/Engine/Source/Developer/MaterialBaking/Private/MaterialBakingModule.cpp b/Engine/Source/Developer/MaterialBaking/Private/MaterialBakingModule.cpp
--- a/Engine/Source/Developer/MaterialBaking/Private/MaterialBakingModule.cpp
+++ b/Engine/Source/Developer/MaterialBaking/Private/MaterialBakingModule.cpp
@@ -21,6 +21,7 @@
 #include "RenderingThread.h"
 #include "RHISurfaceDataConversion.h"
 #include "SceneView.h"
+#include "Serialization/ArchiveCrc32.h"
 #include "Misc/ScopedSlowTask.h"
 #include "MeshDescription.h"
 #include "TextureCompiler.h"
@@ -325,6 +326,29 @@ void FMaterialBakingModule::ShutdownModule()
 	CleanupMaterialProxies();
 }
 
+uint32 FMaterialBakingModule::GetCRC() const
+{
+	FArchiveCrc32 Ar;
+
+	// Base key, changing this will force a rebuild of all HLODs that are relying on material baking.
+	FString ModuleBaseKey = "4167B9A126CA47B3A6EAB520B40A66BB";
+	Ar << ModuleBaseKey;
+
+	bool bUseEmissiveHDR = bEmissiveHDR;
+	Ar << bUseEmissiveHDR;
+
+	uint8 ColorSpace = DefaultColorSpace;
+	Ar << ColorSpace;
+
+	int32 VTWarmupFrames = CVarMaterialBakingVTWarmupFrames.GetValueOnAnyThread();
+	Ar << VTWarmupFrames;
+
+	bool ForceDisableEmissiveScaling = CVarMaterialBakingForceDisableEmissiveScaling.GetValueOnAnyThread();
+	Ar << ForceDisableEmissiveScaling;
+
+	return Ar.GetCrc();
+}
+
 void FMaterialBakingModule::BakeMaterials(const TArray<FMaterialData*>& MaterialSettings, const TArray<FMeshData*>& MeshSettings, TArray<FBakeOutput>& Output)
 {
 	// Translate old material data to extended types
diff --git a/Engine/Source/Developer/MaterialBaking/Private/MaterialBakingModule.h b/Engine/Source/Developer/MaterialBaking/Private/MaterialBakingModule.h
--- a/Engine/Source/Developer/MaterialBaking/Private/MaterialBakingModule.h
+++ b/Engine/Source/Developer/MaterialBaking/Private/MaterialBakingModule.h
@@ -46,6 +46,9 @@ public:
 	/** Returns whether a specific material property is baked to a linear texture or not */
 	virtual bool IsLinearBake(FMaterialPropertyEx Property) override;
 
+	/** Obtain a CRC than can help trigger a rebake if code/global settings impacting the bake result change */
+	virtual uint32 GetCRC() const override;
+
 protected:
 	/* Creates and adds or reuses a RenderTarget from the pool */
 	UTextureRenderTarget2D* CreateRenderTarget(bool bInForceLinearGamma, EPixelFormat InPixelFormat, const FIntPoint& InTargetSize, const FColor& BackgroundColor);
diff --git a/Engine/Source/Developer/MaterialBaking/Public/IMaterialBakingModule.h b/Engine/Source/Developer/MaterialBaking/Public/IMaterialBakingModule.h
--- a/Engine/Source/Developer/MaterialBaking/Public/IMaterialBakingModule.h
+++ b/Engine/Source/Developer/MaterialBaking/Public/IMaterialBakingModule.h
@@ -41,4 +41,7 @@ public:
 
 	/** Returns whether a specific material property is baked to a linear texture or not */
 	virtual bool IsLinearBake(FMaterialPropertyEx Property) = 0;
+
+	/** Obtain a CRC than can help trigger a rebake if code/global settings impacting the bake result change */
+	virtual uint32 GetCRC() const = 0;
 };
