Material Baking - When force streaming in all the referenced textures mips, use a better(faster) mechanism to wait for completion

diff --git a/Engine/Source/Developer/MaterialBaking/Private/MaterialBakingModule.cpp b/Engine/Source/Developer/MaterialBaking/Private/MaterialBakingModule.cpp
--- a/Engine/Source/Developer/MaterialBaking/Private/MaterialBakingModule.cpp
+++ b/Engine/Source/Developer/MaterialBaking/Private/MaterialBakingModule.cpp
@@ -1,6 +1,7 @@
 // Copyright Epic Games, Inc. All Rights Reserved.
 
 #include "MaterialBakingModule.h"
+#include "ContentStreaming.h"
 #include "MaterialRenderItem.h"
 #include "Engine/TextureRenderTarget2D.h"
 #include "ExportMaterialProxy.h"
@@ -516,13 +517,13 @@ void FMaterialBakingModule::BakeMaterials(const TArray<FMaterialDataEx*>& Materi
 
 				for (UTexture* Texture : MaterialTextures)
 				{
-					if (Texture != NULL)
+					if (UTexture2D* Texture2D = Cast<UTexture2D>(Texture))
 					{
-						UTexture2D* Texture2D = Cast<UTexture2D>(Texture);
-						if (Texture2D)
+						if (Texture2D->IsStreamable())
 						{
-							Texture2D->SetForceMipLevelsToBeResident(30.0f);
-							Texture2D->WaitForStreaming();
+							// Force LODs in with high priority, including cinematic ones.
+							Texture2D->SetForceMipLevelsToBeResident(30.f, 0xFFFFFFFF);
+							Texture2D->StreamIn(FStreamableRenderResourceState::MAX_LOD_COUNT, true);
 						}
 					}
 				}
@@ -534,6 +535,18 @@ void FMaterialBakingModule::BakeMaterials(const TArray<FMaterialDataEx*>& Materi
 				CreateMaterialProxy(CurrentMaterialSettings, PropertySizeIterator.Key());
 			}
 		}
+
+		// Force all mip maps to load before baking the materials
+		{
+			TRACE_CPUPROFILER_EVENT_SCOPE(ForceUpdateTextureStreaming)
+			UTexture::ForceUpdateTextureStreaming();
+		}
+
+		// Force all streamed resources to finish
+		{
+			TRACE_CPUPROFILER_EVENT_SCOPE(StreamAllResources)
+			IStreamingManager::Get().StreamAllResources(0.0f);
+		}
 	}
 
 	TAtomic<uint32>    NumTasks(0);
