Material Baking - Added optimized code path when baking a set of meshes with remapped UVs
* In this case, instead of baking each mesh material property to it's own buffer, we can write all the meshes to the same buffer as their UVs are already the final ones that do not overlap
* Memory usage is greatly reduced (in test case with ~300 materials to bake, peak is reduced by 5.5GB)
* A lot faster to process (~2x in most cases) as we are doing a single RT readback per baked material property
* We can also skip the final merge of the baked outputs, as we have a single one

Refactored the material baking module to have both code paths (single output and multi outputs) share as much code as possible

diff --git a/Engine/Source/Developer/MaterialBaking/Private/MaterialBakingModule.cpp b/Engine/Source/Developer/MaterialBaking/Private/MaterialBakingModule.cpp
--- a/Engine/Source/Developer/MaterialBaking/Private/MaterialBakingModule.cpp
+++ b/Engine/Source/Developer/MaterialBaking/Private/MaterialBakingModule.cpp
@@ -73,7 +73,7 @@ static TAutoConsoleVariable<bool> CVarMaterialBakingForceDisableEmissiveScaling(
 	false,
 	TEXT("If set to true, values stored in the emissive textures will be clamped to the [0, 1] range rather than being normalized and scaled back using the EmissiveScale material static parameter."));
 
-namespace FMaterialBakingModuleImpl
+namespace
 {
 	// Custom dynamic mesh allocator specifically tailored for Material Baking.
 	// This will always reuse the same couple buffers, so searching linearly is not a problem.
@@ -273,6 +273,9 @@ namespace FMaterialBakingModuleImpl
 	}
 }
 
+/** Helper for emissive color conversion to Output */
+static void ProcessEmissiveOutput(const FFloat16Color* Color16, int32 Color16Pitch, const FIntPoint& OutputSize, TArray<FColor>& Output, float& EmissiveScale, const FColor& BackgroundColor);
+
 void FMaterialBakingModule::StartupModule()
 {
 	bEmissiveHDR = false;
@@ -315,7 +318,6 @@ void FMaterialBakingModule::ShutdownModule()
 {
 	// Unregister customization and callback
 	FPropertyEditorModule* PropertyEditorModule = FModuleManager::GetModulePtr<FPropertyEditorModule>("PropertyEditor");
-
 	if (PropertyEditorModule)
 	{
 		PropertyEditorModule->UnregisterCustomPropertyTypeLayout(TEXT("PropertyEntry"));
@@ -325,6 +327,8 @@ void FMaterialBakingModule::ShutdownModule()
 	FCoreUObjectDelegates::GetPreGarbageCollectDelegate().RemoveAll(this);
 
 	CleanupMaterialProxies();
+
+	CleanupRenderTargets();
 }
 
 uint32 FMaterialBakingModule::GetCRC() const
@@ -350,6 +354,46 @@ uint32 FMaterialBakingModule::GetCRC() const
 	return Ar.GetCrc();
 }
 
+FMaterialDataEx ToMaterialDataEx(const FMaterialData& MaterialData)
+{
+	FMaterialDataEx MaterialDataEx;
+	MaterialDataEx.Material = MaterialData.Material;
+	MaterialDataEx.bPerformBorderSmear = MaterialData.bPerformBorderSmear;
+	MaterialDataEx.bPerformShrinking = MaterialData.bPerformShrinking;
+	MaterialDataEx.bTangentSpaceNormal = MaterialData.bTangentSpaceNormal;
+	MaterialDataEx.BlendMode = MaterialData.BlendMode;
+	MaterialDataEx.BackgroundColor = MaterialData.BackgroundColor;
+	for (const TPair<EMaterialProperty, FIntPoint>& PropertySizePair : MaterialData.PropertySizes)
+	{
+		MaterialDataEx.PropertySizes.Add(PropertySizePair.Key, PropertySizePair.Value);
+	}
+	return MaterialDataEx;
+}
+
+FBakeOutput ToBakeOutput(FBakeOutputEx& BakeOutputEx)
+{
+	FBakeOutput BakeOutput;
+
+	BakeOutput.EmissiveScale = BakeOutputEx.EmissiveScale;
+
+	for (TPair<FMaterialPropertyEx, FIntPoint>& PropertySizePair : BakeOutputEx.PropertySizes)
+	{
+		BakeOutput.PropertySizes.Add(PropertySizePair.Key.Type, PropertySizePair.Value);
+	}
+
+	for (TPair<FMaterialPropertyEx, TArray<FColor>>& PropertyDataPair : BakeOutputEx.PropertyData)
+	{
+		BakeOutput.PropertyData.Add(PropertyDataPair.Key.Type, MoveTemp(PropertyDataPair.Value));
+	}
+
+	for (TPair<FMaterialPropertyEx, TArray<FFloat16Color>>& PropertyDataPair : BakeOutputEx.HDRPropertyData)
+	{
+		BakeOutput.HDRPropertyData.Add(PropertyDataPair.Key.Type, MoveTemp(PropertyDataPair.Value));
+	}
+
+	return BakeOutput;
+}
+
 void FMaterialBakingModule::BakeMaterials(const TArray<FMaterialData*>& MaterialSettings, const TArray<FMeshData*>& MeshSettings, TArray<FBakeOutput>& Output)
 {
 	// Translate old material data to extended types
@@ -357,18 +401,7 @@ void FMaterialBakingModule::BakeMaterials(const TArray<FMaterialData*>& Material
 	MaterialDataExs.Reserve(MaterialSettings.Num());
 	for (const FMaterialData* MaterialData : MaterialSettings)
 	{
-		FMaterialDataEx& MaterialDataEx = MaterialDataExs.AddDefaulted_GetRef();
-		MaterialDataEx.Material = MaterialData->Material;
-		MaterialDataEx.bPerformBorderSmear = MaterialData->bPerformBorderSmear;
-		MaterialDataEx.bPerformShrinking = MaterialData->bPerformShrinking;
-		MaterialDataEx.bTangentSpaceNormal = MaterialData->bTangentSpaceNormal;
-		MaterialDataEx.BlendMode = MaterialData->BlendMode;
-		MaterialDataEx.BackgroundColor = MaterialData->BackgroundColor;
-
-		for (const TPair<EMaterialProperty, FIntPoint>& PropertySizePair : MaterialData->PropertySizes)
-		{
-			MaterialDataEx.PropertySizes.Add(PropertySizePair.Key, PropertySizePair.Value);
-		}
+		MaterialDataExs.Emplace(ToMaterialDataEx(*MaterialData));
 	}
 
 	// Build an array of pointers to the extended type
@@ -386,118 +419,236 @@ void FMaterialBakingModule::BakeMaterials(const TArray<FMaterialData*>& Material
 	Output.Reserve(BakeOutputExs.Num());
 	for (FBakeOutputEx& BakeOutputEx : BakeOutputExs)
 	{
-		FBakeOutput& BakeOutput = Output.AddDefaulted_GetRef();
-		BakeOutput.EmissiveScale = BakeOutputEx.EmissiveScale;
-
-		for (TPair<FMaterialPropertyEx, FIntPoint>& PropertySizePair : BakeOutputEx.PropertySizes)
-		{
-			BakeOutput.PropertySizes.Add(PropertySizePair.Key.Type, PropertySizePair.Value);
-		}
+		Output.Emplace(ToBakeOutput(BakeOutputEx));
+	}
+}
 
-		for (TPair<FMaterialPropertyEx, TArray<FColor>>& PropertyDataPair : BakeOutputEx.PropertyData)
-		{
-			BakeOutput.PropertyData.Add(PropertyDataPair.Key.Type, MoveTemp(PropertyDataPair.Value));
-		}
+void FMaterialBakingModule::BakeMaterials(const TArray<FMaterialData*>& MaterialSettings, const TArray<FMeshData*>& MeshSettings, FBakeOutput& BakeOutput)
+{
+	// Translate old material data to extended types
+	TArray<FMaterialDataEx> MaterialDataExs;
+	MaterialDataExs.Reserve(MaterialSettings.Num());
+	for (const FMaterialData* MaterialData : MaterialSettings)
+	{
+		MaterialDataExs.Emplace(ToMaterialDataEx(*MaterialData));
+	}
 
-		for (TPair<FMaterialPropertyEx, TArray<FFloat16Color>>& PropertyDataPair : BakeOutputEx.HDRPropertyData)
-		{
-			BakeOutput.HDRPropertyData.Add(PropertyDataPair.Key.Type, MoveTemp(PropertyDataPair.Value));
-		}
+	// Build an array of pointers to the extended type
+	TArray<FMaterialDataEx*> MaterialSettingsEx;
+	MaterialSettingsEx.Reserve(MaterialDataExs.Num());
+	for (FMaterialDataEx& MaterialDataEx : MaterialDataExs)
+	{
+		MaterialSettingsEx.Add(&MaterialDataEx);
 	}
+
+	FBakeOutputEx BakeOutputEx;
+	BakeMaterials(MaterialSettingsEx, MeshSettings, BakeOutputEx);
+
+	// Translate extended bake output to old type
+	BakeOutput = ToBakeOutput(BakeOutputEx);
 }
 
-void FMaterialBakingModule::BakeMaterials(const TArray<FMaterialDataEx*>& MaterialSettings, const TArray<FMeshData*>& MeshSettings, TArray<FBakeOutputEx>& Output)
+
+class FMaterialBakingProcessor
 {
-	UE_LOG(LogMaterialBaking, Verbose, TEXT("Performing material baking for %d materials"), MaterialSettings.Num());
-	for (int32 i = 0; i < MaterialSettings.Num(); i++)
+public:
+	FMaterialBakingProcessor(FMaterialBakingModule& InMaterialBakingModule, const TArray<FMaterialDataEx*>& InMaterialSettings, const TArray<FMeshData*>& InMeshSettings)
+		: MaterialBakingModule(InMaterialBakingModule)
+		, MaterialSettings(InMaterialSettings)
+		, MeshSettings(InMeshSettings)
+		, NumMaterials(MaterialSettings.Num())
+		, bSaveIntermediateTextures(CVarSaveIntermediateTextures.GetValueOnAnyThread() == 1)
+		, bEmissiveHDR(InMaterialBakingModule.bEmissiveHDR)
 	{
-		if (MaterialSettings[i]->Material && MeshSettings[i]->MeshDescription)
+		checkf(MaterialSettings.Num() == MeshSettings.Num(), TEXT("Number of material settings does not match that of MeshSettings"));
+
+		UE_LOG(LogMaterialBaking, Verbose, TEXT("Performing material baking for %d materials"), MaterialSettings.Num());
+		for (int32 i = 0; i < MaterialSettings.Num(); i++)
 		{
-			UE_LOG(LogMaterialBaking, Verbose, TEXT("    [%5d] Material: %-50s Vertices: %8d    Triangles: %8d"), i, *MaterialSettings[i]->Material->GetName(), MeshSettings[i]->MeshDescription->Vertices().Num(), MeshSettings[i]->MeshDescription->Triangles().Num());
+			if (MaterialSettings[i]->Material && MeshSettings[i]->MeshDescription)
+			{
+				UE_LOG(LogMaterialBaking, Verbose, TEXT("    [%5d] Material: %-50s Vertices: %8d    Triangles: %8d"), i, *MaterialSettings[i]->Material->GetName(), MeshSettings[i]->MeshDescription->Vertices().Num(), MeshSettings[i]->MeshDescription->Triangles().Num());
+			}
 		}
+
+		ComputeMeshProcessingOrder();
 	}
 
-	TRACE_CPUPROFILER_EVENT_SCOPE(FMaterialBakingModule::BakeMaterials)
+	virtual ~FMaterialBakingProcessor() = default;
 
-	checkf(MaterialSettings.Num() == MeshSettings.Num(), TEXT("Number of material settings does not match that of MeshSettings"));
-	const int32 NumMaterials = MaterialSettings.Num();
-	const bool bSaveIntermediateTextures = CVarSaveIntermediateTextures.GetValueOnAnyThread() == 1;
+	void BakeMaterials()
+	{
+		TRACE_CPUPROFILER_EVENT_SCOPE(FMaterialBakingModule::BakeMaterials)
 
-	using namespace FMaterialBakingModuleImpl;
-	FMaterialBakingDynamicMeshBufferAllocator MaterialBakingDynamicMeshBufferAllocator;
+		// We reuse the pipeline depth to prepare render items in advance to avoid stalling the game thread
+		LaunchAsyncPrepareRenderItems(PipelineDepth);
 
-	FScopedSlowTask Progress(NumMaterials, LOCTEXT("BakeMaterials", "Baking Materials..."), true );
-	Progress.MakeDialog(true);
+		// Create all material proxies right away to start compiling shaders asynchronously and avoid stalling the baking process as much as possible
+		CreateMaterialProxies();
 
-	TArray<uint32> ProcessingOrder;
-	ProcessingOrder.Reserve(MeshSettings.Num());
-	for (int32 Index = 0; Index < MeshSettings.Num(); ++Index)
+		// For each material
+		for (int32 Index = 0; Index < NumMaterials; ++Index)
+		{
+			const int32 MaterialIndex = ProcessingOrder[Index];
+			const FMaterialDataEx& CurrentMaterialSettings = *MaterialSettings[MaterialIndex];
+			const FMeshData* CurrentMeshSettings = MeshSettings[MaterialIndex];
+			FBakeOutputEx& CurrentOutput = GetBakeOutput(MaterialIndex);
+
+			TMap<FRenderItemKey, FMeshMaterialRenderItem*>* RenderItems = GetRenderItems(Index);
+			check(RenderItems && !RenderItems->IsEmpty());
+
+			// For each property
+			for (const auto& [Property, Size] : CurrentMaterialSettings.PropertySizes)
+			{
+				TRACE_CPUPROFILER_EVENT_SCOPE_TEXT(*Property.ToString())
+
+				FExportMaterialProxy* ExportMaterialProxy = MaterialBakingModule.CreateMaterialProxy(&CurrentMaterialSettings, Property);
+				if (!ExportMaterialProxy->IsCompilationFinished())
+				{
+					TRACE_CPUPROFILER_EVENT_SCOPE(WaitForMaterialProxyCompilation)
+					ExportMaterialProxy->FinishCompilation();
+				}
+
+				UTextureRenderTarget2D* RenderTarget = GetRenderTarget(Property, Size, CurrentMaterialSettings);
+				FMeshMaterialRenderItem* RenderItem = RenderItems->FindChecked(FRenderItemKey(CurrentMeshSettings, Size));
+
+				BakeMaterialProperty(CurrentMaterialSettings, Property, RenderItem, RenderTarget, ExportMaterialProxy, CurrentOutput);
+			}
+
+			// Destroying Render Items
+			// Must happen on the render thread to ensure they are not used anymore.
+			ENQUEUE_RENDER_COMMAND(DestroyRenderItems)(
+				[RenderItems](FRHICommandListImmediate& RHICmdList)
+				{
+					for (auto RenderItem : (*RenderItems))
+					{
+						delete RenderItem.Value;
+					}
+
+					delete RenderItems;
+				}
+			);
+		}
+	}
+
+protected:
+	void PrepareBakeOutput(FMaterialDataEx* InMaterialSettings, FBakeOutputEx& BakeOutput)
 	{
-		ProcessingOrder.Add(Index);
+		BakeOutput.PropertySizes = InMaterialSettings->PropertySizes;
+
+		for (const auto& [Property, Size] : BakeOutput.PropertySizes)
+		{
+			BakeOutput.PropertyData.Add(Property);
+			if (bEmissiveHDR && Property == MP_EmissiveColor)
+			{
+				BakeOutput.HDRPropertyData.Add(Property);
+			}
+		}
 	}
 
-	// Start with the biggest mesh first so we can always reuse the same vertex/index buffers.
-	// This will decrease the number of allocations backed by newly allocated memory from the OS,
-	// which will reduce soft page faults while copying into that memory.
-	// Soft page faults are now incredibly expensive on Windows 10.
-	Algo::SortBy(
-		ProcessingOrder,
-		[&MeshSettings](const uint32 Index){ return MeshSettings[Index]->MeshDescription ? MeshSettings[Index]->MeshDescription->Vertices().Num() : 0; },
-		TGreater<>()
-	);
+	UTextureRenderTarget2D* CreateRenderTarget(FMaterialPropertyEx InProperty, const FIntPoint& InTargetSize, bool bInUsePooledRenderTargets, const FColor& BackgroundColor)
+	{
+		return MaterialBakingModule.CreateRenderTarget(InProperty, InTargetSize, bInUsePooledRenderTargets, BackgroundColor);
+	}
 
-	Output.SetNum(NumMaterials);
+	void SaveIntermediateTextures(const FBakeOutputEx& BakeOutput, const FMaterialPropertyEx& Property, const FString& FilenameString)
+	{
+#if WITH_EDITOR
+		// If saving intermediates is turned on
+		if (bSaveIntermediateTextures)
+		{
+			if (!BakeOutput.PropertyData[Property].IsEmpty())
+			{
+				static int32 SaveCount = 0;
+				TRACE_CPUPROFILER_EVENT_SCOPE(SaveIntermediateTextures)
+				FString TrimmedPropertyName = Property.ToString();
 
-	struct FPipelineContext
+				const FString DirectoryPath = FPaths::ConvertRelativePathToFull(FPaths::ProjectIntermediateDir() + TEXT("MaterialBaking/"));
+				FString FullPath = FString::Printf(TEXT("%s%s-%d-%s.bmp"), *DirectoryPath, *FilenameString, SaveCount++, *TrimmedPropertyName);
+				FFileHelper::CreateBitmap(*FullPath, BakeOutput.PropertySizes[Property].X, BakeOutput.PropertySizes[Property].Y, BakeOutput.PropertyData[Property].GetData());
+			}
+		}
+#endif
+	}
+
+private:
+	virtual FBakeOutputEx& GetBakeOutput(int32 InMaterialIndex) = 0;
+	virtual UTextureRenderTarget2D* GetRenderTarget(FMaterialPropertyEx InMaterialProperty, const FIntPoint& InRequiredSize, const FMaterialDataEx& InMaterialSettings) = 0;
+	virtual void OnMaterialPropertyBaked(const FMaterialDataEx& CurrentMaterialSettings, const FMaterialPropertyEx& Property, FMeshMaterialRenderItem* RenderItem, UTextureRenderTarget2D* RenderTarget, FExportMaterialProxy* ExportMaterialProxy, FBakeOutputEx& CurrentOutput) {}
+
+	void ComputeMeshProcessingOrder()
 	{
-		typedef TFunction<void (FRHICommandListImmediate& RHICmdList)> FReadCommand;
-		FReadCommand ReadCommand;
-	};
+		ProcessingOrder.Reserve(MeshSettings.Num());
+		for (int32 Index = 0; Index < MeshSettings.Num(); ++Index)
+		{
+			ProcessingOrder.Add(Index);
+		}
 
-	// Distance between the command sent to rendering and the GPU read-back of the result
-	// to minimize sync time waiting on GPU.
-	const int32 PipelineDepth = 16;
-	int32 PipelineIndex = 0;
-	FPipelineContext PipelineContext[PipelineDepth];
+		// Start with the biggest mesh first so we can always reuse the same vertex/index buffers.
+		// This will decrease the number of allocations backed by newly allocated memory from the OS,
+		// which will reduce soft page faults while copying into that memory.
+		// Soft page faults are now incredibly expensive on Windows 10.
+		Algo::SortBy(
+			ProcessingOrder,
+			[this](const uint32 Index) { return MeshSettings[Index]->MeshDescription ? MeshSettings[Index]->MeshDescription->Vertices().Num() : 0; },
+			TGreater<>()
+		);
+	}
 
 	// This will create and prepare FMeshMaterialRenderItem for each property sizes we're going to need
-	auto PrepareRenderItems_AnyThread =
-		[&](int32 MaterialIndex)
+	TMap<FRenderItemKey, FMeshMaterialRenderItem*>* PrepareRenderItems_AnyThread(int32 MaterialIndex)
 	{
 		TRACE_CPUPROFILER_EVENT_SCOPE(PrepareRenderItems);
-		
-		TMap<FMaterialBakingModuleImpl::FRenderItemKey, FMeshMaterialRenderItem*>* RenderItems = new TMap<FRenderItemKey, FMeshMaterialRenderItem *>();
+
+		TMap<FRenderItemKey, FMeshMaterialRenderItem*>* RenderItems = new TMap<FRenderItemKey, FMeshMaterialRenderItem*>();
 		const FMaterialDataEx* CurrentMaterialSettings = MaterialSettings[MaterialIndex];
 		const FMeshData* CurrentMeshSettings = MeshSettings[MaterialIndex];
 
-		for (TMap<FMaterialPropertyEx, FIntPoint>::TConstIterator PropertySizeIterator = CurrentMaterialSettings->PropertySizes.CreateConstIterator(); PropertySizeIterator; ++PropertySizeIterator)
+		for (const auto& [Property, Size] : CurrentMaterialSettings->PropertySizes)
 		{
-			FRenderItemKey RenderItemKey(CurrentMeshSettings, PropertySizeIterator.Value());
+			FRenderItemKey RenderItemKey(CurrentMeshSettings, Size);
 			if (RenderItems->Find(RenderItemKey) == nullptr)
 			{
-				RenderItems->Add(RenderItemKey, new FMeshMaterialRenderItem(PropertySizeIterator.Value(), CurrentMeshSettings, &MaterialBakingDynamicMeshBufferAllocator));
+				RenderItems->Add(RenderItemKey, new FMeshMaterialRenderItem(Size, CurrentMeshSettings, &MaterialBakingDynamicMeshBufferAllocator));
 			}
 		}
 
+		check(!RenderItems->IsEmpty());
 		return RenderItems;
-	};
+	}
 
-	// We reuse the pipeline depth to prepare render items in advance to avoid stalling the game thread
-	int NextRenderItem = 0;
-	TFuture<TMap<FRenderItemKey, FMeshMaterialRenderItem*>*> PreparedRenderItems[PipelineDepth];
-	for (; NextRenderItem < NumMaterials && NextRenderItem < PipelineDepth; ++NextRenderItem)
+	void LaunchAsyncPrepareRenderItems(int32 InLaunchCount)
 	{
-		PreparedRenderItems[NextRenderItem] = 
-			Async(
-				EAsyncExecution::ThreadPool,
-				[&PrepareRenderItems_AnyThread, &ProcessingOrder, NextRenderItem]()
-				{
-					return PrepareRenderItems_AnyThread(ProcessingOrder[NextRenderItem]);
-				}
+		int32 LastRenderItem = NextRenderItem + InLaunchCount;
+		for (; NextRenderItem < NumMaterials && NextRenderItem < LastRenderItem; NextRenderItem++)
+		{
+			PreparedRenderItems[NextRenderItem % PipelineDepth] =
+				Async(
+					EAsyncExecution::ThreadPool,
+					[this, NextRenderItem=NextRenderItem]()
+					{
+						return PrepareRenderItems_AnyThread(ProcessingOrder[NextRenderItem]);
+					}
 			);
+		}
 	}
 
-	// Create all material proxies right away to start compiling shaders asynchronously and avoid stalling the baking process as much as possible
+	TMap<FRenderItemKey, FMeshMaterialRenderItem*>* GetRenderItems(int32 InIndex)
+	{
+		TRACE_CPUPROFILER_EVENT_SCOPE(WaitOnPreparedRenderItems)
+		TMap<FRenderItemKey, FMeshMaterialRenderItem*>* RenderItems = PreparedRenderItems[InIndex % PipelineDepth].Get();
+
+		// Prepare the next render item in advance
+		if (NextRenderItem < NumMaterials)
+		{
+			check((NextRenderItem % PipelineDepth) == (InIndex % PipelineDepth));
+			LaunchAsyncPrepareRenderItems(1);
+		}
+
+		return RenderItems;
+	}
+
+	void CreateMaterialProxies()
 	{
 		TRACE_CPUPROFILER_EVENT_SCOPE(CreateMaterialProxies)
 
@@ -532,273 +683,212 @@ void FMaterialBakingModule::BakeMaterials(const TArray<FMaterialDataEx*>& Materi
 			for (TMap<FMaterialPropertyEx, FIntPoint>::TConstIterator PropertySizeIterator = CurrentMaterialSettings->PropertySizes.CreateConstIterator(); PropertySizeIterator; ++PropertySizeIterator)
 			{
 				// They will be stored in the pool and compiled asynchronously
-				CreateMaterialProxy(CurrentMaterialSettings, PropertySizeIterator.Key());
+				MaterialBakingModule.CreateMaterialProxy(CurrentMaterialSettings, PropertySizeIterator.Key());
 			}
 		}
 
 		// Force all mip maps to load before baking the materials
-		{
-			TRACE_CPUPROFILER_EVENT_SCOPE(ForceUpdateTextureStreaming)
-			UTexture::ForceUpdateTextureStreaming();
-		}
-
-		// Force all streamed resources to finish
-		{
-			TRACE_CPUPROFILER_EVENT_SCOPE(StreamAllResources)
-			IStreamingManager::Get().StreamAllResources(0.0f);
-		}
+		UTexture::ForceUpdateTextureStreaming();
 	}
 
-	TAtomic<uint32>    NumTasks(0);
-	FStagingBufferPool StagingBufferPool;
-
-	for (int32 Index = 0; Index < NumMaterials; ++Index)
+	void BakeMaterialProperty(const FMaterialDataEx& CurrentMaterialSettings, const FMaterialPropertyEx& Property, FMeshMaterialRenderItem* RenderItem, UTextureRenderTarget2D* RenderTarget, FExportMaterialProxy* ExportMaterialProxy, FBakeOutputEx& CurrentOutput)
 	{
-		TRACE_CPUPROFILER_EVENT_SCOPE(BakeOneMaterial)
+		// Perform everything left of the operation directly on the render thread since we need to modify some RenderItem's properties
+		// for each render pass and we can't do that without costly synchronization (flush) between the game thread and render thread.
+		// Everything slow to execute has already been prepared on the game thread anyway.
+		ENQUEUE_RENDER_COMMAND(RenderOneMaterial)(
+			[RenderItem, RenderTarget, ExportMaterialProxy](FRHICommandListImmediate& RHICmdList)
+			{
+				FSceneViewFamily ViewFamily(FSceneViewFamily::ConstructionValues(RenderTarget->GetRenderTargetResource(), nullptr,
+					FEngineShowFlags(ESFIM_Game))
+					.SetTime(FGameTime())
+					.SetGammaCorrection(RenderTarget->GetRenderTargetResource()->GetDisplayGamma()));
+
+				RenderItem->MaterialRenderProxy = ExportMaterialProxy;
+				RenderItem->ViewFamily = &ViewFamily;
+
+				FTextureRenderTargetResource* RenderTargetResource = RenderTarget->GetRenderTargetResource();
+				FCanvas Canvas(RenderTargetResource, nullptr, FGameTime::GetTimeSinceAppStart(), GMaxRHIFeatureLevel);
+				Canvas.SetAllowedModes(FCanvas::Allow_Flush);
+				Canvas.SetRenderTargetRect(FIntRect(0, 0, RenderTarget->GetSurfaceWidth(), RenderTarget->GetSurfaceHeight()));
+				Canvas.SetBaseTransform(Canvas.CalcBaseTransform2D(RenderTarget->GetSurfaceWidth(), RenderTarget->GetSurfaceHeight()));
+
+				// Virtual textures may require repeated rendering to warm up.
+				int32 WarmupIterationCount = 1;
+				if (UseVirtualTexturing(ViewFamily.GetFeatureLevel()))
+				{
+					const FMaterial& MeshMaterial = ExportMaterialProxy->GetIncompleteMaterialWithFallback(ViewFamily.GetFeatureLevel());
+					if (!MeshMaterial.GetUniformVirtualTextureExpressions().IsEmpty())
+					{
+						WarmupIterationCount = CVarMaterialBakingVTWarmupFrames.GetValueOnAnyThread();
+					}
+				}
 
-		Progress.EnterProgressFrame(1.0f, FText::Format(LOCTEXT("BakingMaterial", "Baking Material {0}/{1}"), Index, NumMaterials));
+				// Do rendering
+				for (int WarmupIndex = 0; WarmupIndex < WarmupIterationCount; ++WarmupIndex)
+				{
+					FCanvas::FCanvasSortElement& SortElement = Canvas.GetSortElement(Canvas.TopDepthSortKey());
+					SortElement.RenderBatchArray.Add(RenderItem);
+					Canvas.Flush_RenderThread(RHICmdList);
+					SortElement.RenderBatchArray.Empty();
 
-		int32 MaterialIndex = ProcessingOrder[Index];
-		TMap<FRenderItemKey, FMeshMaterialRenderItem*>* RenderItems;
+					RHICmdList.ImmediateFlush(EImmediateFlushType::FlushRHIThreadFlushResources);
+				}
+			}
+		);
 
-		{
-			TRACE_CPUPROFILER_EVENT_SCOPE(WaitOnPreparedRenderItems)
-			RenderItems = PreparedRenderItems[Index % PipelineDepth].Get();
-		}
+		OnMaterialPropertyBaked(CurrentMaterialSettings, Property, RenderItem, RenderTarget, ExportMaterialProxy, CurrentOutput);
+	}
 
-		// Prepare the next render item in advance
-		if (NextRenderItem < NumMaterials)
-		{
-			check((NextRenderItem % PipelineDepth) == (Index % PipelineDepth));
-			PreparedRenderItems[NextRenderItem % PipelineDepth] = 
-				Async(
-					EAsyncExecution::ThreadPool,
-					[&PrepareRenderItems_AnyThread, NextMaterialIndex = ProcessingOrder[NextRenderItem]]()
-					{
-						return PrepareRenderItems_AnyThread(NextMaterialIndex);
-					}
-				);
-			NextRenderItem++;
-		}
+protected:
+	FMaterialBakingModule& MaterialBakingModule;
 
-		const FMaterialDataEx* CurrentMaterialSettings = MaterialSettings[MaterialIndex];
-		const FMeshData* CurrentMeshSettings = MeshSettings[MaterialIndex];
-		FBakeOutputEx& CurrentOutput = Output[MaterialIndex];
+	const TArray<FMaterialDataEx*>& MaterialSettings;
+	const TArray<FMeshData*>& MeshSettings;
+	const int32 NumMaterials;
+	const bool bSaveIntermediateTextures;
+	const bool bEmissiveHDR;
+
+	// Distance between the command sent to rendering and the GPU read-back of the result
+	// to minimize sync time waiting on GPU.
+	static const int32 PipelineDepth = 16;
+
+private:
+	TArray<uint32> ProcessingOrder;
 
-		TRACE_CPUPROFILER_EVENT_SCOPE_TEXT(*CurrentMaterialSettings->Material->GetName())
+	FMaterialBakingDynamicMeshBufferAllocator MaterialBakingDynamicMeshBufferAllocator;
+
+	// We reuse the pipeline depth to prepare render items in advance to avoid stalling the game thread
+	int NextRenderItem = 0;
+	TFuture<TMap<FRenderItemKey, FMeshMaterialRenderItem*>*> PreparedRenderItems[PipelineDepth];
+};
 
-		TArray<FMaterialPropertyEx> MaterialPropertiesToBakeOut;
-		CurrentMaterialSettings->PropertySizes.GenerateKeyArray(MaterialPropertiesToBakeOut);
 
-		const int32 NumPropertiesToRender = MaterialPropertiesToBakeOut.Num();
-		if (NumPropertiesToRender > 0)
+class FMaterialBakingProcessorSingleOutput : public FMaterialBakingProcessor
+{
+public:
+	FMaterialBakingProcessorSingleOutput(FMaterialBakingModule& InMaterialBakingModule, const TArray<FMaterialDataEx*>& InMaterialSettings, const TArray<FMeshData*>& InMeshSettings, FBakeOutputEx& InOutput)
+		: FMaterialBakingProcessor(InMaterialBakingModule, InMaterialSettings, InMeshSettings)
+		, Output(InOutput)
+	{
+		if (!MaterialSettings.IsEmpty())
 		{
-			TRACE_CPUPROFILER_EVENT_SCOPE(RenderOneMaterial)
+			FMaterialDataEx* DefaultMaterialData = MaterialSettings[0];
 
-			// Ensure data in memory will not change place passed this point to avoid race conditions
-			CurrentOutput.PropertySizes = CurrentMaterialSettings->PropertySizes;
-			for (int32 PropertyIndex = 0; PropertyIndex < NumPropertiesToRender; ++PropertyIndex)
+			// Single output path can only work if all materials settings share the same properties
+			for (const FMaterialDataEx* MaterialData : MaterialSettings)
 			{
-				const FMaterialPropertyEx& Property = MaterialPropertiesToBakeOut[PropertyIndex];
-				CurrentOutput.PropertyData.Add(Property);
-				if (bEmissiveHDR && Property == MP_EmissiveColor)
+				for (const auto& [Property, Size] : MaterialData->PropertySizes)
 				{
-					CurrentOutput.HDRPropertyData.Add(Property);
+					check(DefaultMaterialData->PropertySizes.Contains(Property) && DefaultMaterialData->PropertySizes[Property] == Size);
 				}
+				check(MaterialData->bPerformBorderSmear == DefaultMaterialData->bPerformBorderSmear);
+				check(MaterialData->bPerformShrinking == DefaultMaterialData->bPerformShrinking);
+				check(MaterialData->bTangentSpaceNormal == DefaultMaterialData->bTangentSpaceNormal);
+				check(MaterialData->BlendMode == DefaultMaterialData->BlendMode);
+				check(MaterialData->BackgroundColor == DefaultMaterialData->BackgroundColor);
 			}
 
-			for (int32 PropertyIndex = 0; PropertyIndex < NumPropertiesToRender; ++PropertyIndex)
+			PrepareBakeOutput(DefaultMaterialData, Output);
+
+			// Create render targets for all properties
+			for (const auto& [Property, Size] : Output.PropertySizes)
 			{
-				const FMaterialPropertyEx& Property = MaterialPropertiesToBakeOut[PropertyIndex];
+				// Skip pooling, all materials will be baked to the same set of render targets
+				const bool bUsePooledRenderTargets = false;
 
-				TRACE_CPUPROFILER_EVENT_SCOPE_TEXT(*Property.ToString())
+				UTextureRenderTarget2D* RenderTarget = CreateRenderTarget(Property, Size, bUsePooledRenderTargets, DefaultMaterialData->BackgroundColor);
+				RenderTargets.Add(Property, RenderTarget);
+			}
+		}
+	}
 
-				FExportMaterialProxy* ExportMaterialProxy = CreateMaterialProxy(CurrentMaterialSettings, Property);
+	~FMaterialBakingProcessorSingleOutput()
+	{
+		if (!MaterialSettings.IsEmpty())
+		{
+			// Wait until every tasks have been queued so that NumTasks is only decreasing
+			FlushRenderingCommands();
 
-				if (!ExportMaterialProxy->IsCompilationFinished())
+			FMaterialDataEx* DefaultMaterialData = MaterialSettings[0];
+
+			// Read back from the render targets
+			for (auto& [Property, RenderTarget] : RenderTargets)
+			{
+				FRenderTarget* RTResource = RenderTarget->GameThread_GetRenderTargetResource();
+				check(RTResource);
+
+				if (Property.Type != MP_EmissiveColor)
 				{
-					TRACE_CPUPROFILER_EVENT_SCOPE(WaitForMaterialProxyCompilation)
-					ExportMaterialProxy->FinishCompilation();
+					RTResource->ReadPixels(Output.PropertyData[Property]);
+				}
+				else
+				{
+					TArray<FFloat16Color> OutputHDRColor;
+					RTResource->ReadFloat16Pixels(OutputHDRColor);
+					
+					int32 Color16Pitch = Output.PropertySizes[Property].X * sizeof(FFloat16Color);
+					ProcessEmissiveOutput(OutputHDRColor.GetData(), Color16Pitch, Output.PropertySizes[Property], Output.PropertyData[Property], Output.EmissiveScale, DefaultMaterialData->BackgroundColor);
+					
+					if (bEmissiveHDR)
+					{
+						Output.HDRPropertyData[Property] = MoveTemp(OutputHDRColor);
+					}
 				}
 
-				// Lookup gamma and format settings for property, if not found use default values
-				const EPropertyColorSpace* OverrideColorSpace = PerPropertyColorSpace.Find(Property);
-				const EPropertyColorSpace ColorSpace = OverrideColorSpace ? *OverrideColorSpace : DefaultColorSpace;
-				const EPixelFormat PixelFormat = PerPropertyFormat.Contains(Property) ? PerPropertyFormat[Property] : PF_B8G8R8A8;
-
-				// It is safe to reuse the same render target for each draw pass since they all execute sequentially on the GPU and are copied to staging buffers before
-				// being reused.
-				UTextureRenderTarget2D* RenderTarget = CreateRenderTarget((ColorSpace == EPropertyColorSpace::Linear), PixelFormat, CurrentOutput.PropertySizes[Property], CurrentMaterialSettings->BackgroundColor);
-				if (RenderTarget != nullptr)
+				if (DefaultMaterialData->bPerformBorderSmear)
 				{
-					// Perform everything left of the operation directly on the render thread since we need to modify some RenderItem's properties
-					// for each render pass and we can't do that without costly synchronization (flush) between the game thread and render thread.
-					// Everything slow to execute has already been prepared on the game thread anyway.
-					ENQUEUE_RENDER_COMMAND(RenderOneMaterial)(
-						[this, RenderItems, RenderTarget, Property, ExportMaterialProxy, &PipelineContext, PipelineIndex, &StagingBufferPool, &NumTasks, bSaveIntermediateTextures, &MaterialSettings, &MeshSettings, MaterialIndex, &Output](FRHICommandListImmediate& RHICmdList)
-						{
-							const FMaterialDataEx* CurrentMaterialSettings = MaterialSettings[MaterialIndex];
-							const FMeshData* CurrentMeshSettings = MeshSettings[MaterialIndex];
-
-							FMeshMaterialRenderItem& RenderItem = *RenderItems->FindChecked(FRenderItemKey(CurrentMeshSettings, FIntPoint(RenderTarget->GetSurfaceWidth(), RenderTarget->GetSurfaceHeight())));
-
-							FSceneViewFamily ViewFamily(FSceneViewFamily::ConstructionValues(RenderTarget->GetRenderTargetResource(), nullptr,
-								FEngineShowFlags(ESFIM_Game))
-								.SetTime(FGameTime())
-								.SetGammaCorrection(RenderTarget->GetRenderTargetResource()->GetDisplayGamma()));
-
-							RenderItem.MaterialRenderProxy = ExportMaterialProxy;
-							RenderItem.ViewFamily = &ViewFamily;
-
-							FTextureRenderTargetResource* RenderTargetResource = RenderTarget->GetRenderTargetResource();
-							FCanvas Canvas(RenderTargetResource, nullptr, FGameTime::GetTimeSinceAppStart(), GMaxRHIFeatureLevel);
-							Canvas.SetAllowedModes(FCanvas::Allow_Flush);
-							Canvas.SetRenderTargetRect(FIntRect(0, 0, RenderTarget->GetSurfaceWidth(), RenderTarget->GetSurfaceHeight()));
-							Canvas.SetBaseTransform(Canvas.CalcBaseTransform2D(RenderTarget->GetSurfaceWidth(), RenderTarget->GetSurfaceHeight()));
-
-							// Virtual textures may require repeated rendering to warm up.
-							int32 WarmupIterationCount = 1;
- 							if (UseVirtualTexturing(ViewFamily.GetFeatureLevel()))
- 							{
- 								const FMaterial& MeshMaterial = ExportMaterialProxy->GetIncompleteMaterialWithFallback(ViewFamily.GetFeatureLevel());
-								if (!MeshMaterial.GetUniformVirtualTextureExpressions().IsEmpty())
-								{
-									WarmupIterationCount = CVarMaterialBakingVTWarmupFrames.GetValueOnAnyThread();
-								}
- 							}
-
-							// Do rendering
-							{
-								RenderCaptureInterface::FScopedCapture RenderCapture(CVarMaterialBakingRDOCCapture.GetValueOnAnyThread() == 1, &RHICmdList, TEXT("MaterialBaking"));
-
-								for (int WarmupIndex = 0; WarmupIndex < WarmupIterationCount; ++WarmupIndex)
-								{
-									Canvas.Clear(RenderTarget->ClearColor);
-									FCanvas::FCanvasSortElement& SortElement = Canvas.GetSortElement(Canvas.TopDepthSortKey());
-									SortElement.RenderBatchArray.Add(&RenderItem);
-									Canvas.Flush_RenderThread(RHICmdList);
-									SortElement.RenderBatchArray.Empty();
-
-									RHICmdList.ImmediateFlush(EImmediateFlushType::FlushRHIThreadFlushResources);
-								}
-							}
+					// This will resize the output to a single pixel if the result is monochrome.
+					FMaterialBakingHelpers::PerformUVBorderSmearAndShrink(Output.PropertyData[Property], Output.PropertySizes[Property].X, Output.PropertySizes[Property].Y, DefaultMaterialData->BackgroundColor);
+				}
 
-							FTexture2DRHIRef StagingBufferRef = StagingBufferPool.CreateStagingBuffer_RenderThread(RHICmdList, RenderTargetResource->GetSizeX(), RenderTargetResource->GetSizeY(), RenderTarget->GetFormat(), RenderTarget->IsSRGB());
-							FGPUFenceRHIRef GPUFence = RHICreateGPUFence(TEXT("MaterialBackingFence"));
-							TransitionAndCopyTexture(RHICmdList, RenderTargetResource->GetRenderTargetTexture(), StagingBufferRef, {});
-							RHICmdList.WriteGPUFence(GPUFence);
-
-							// Prepare a lambda for final processing that will be executed asynchronously
-							NumTasks++;
-							auto FinalProcessing_AnyThread =
-								[&NumTasks, bSaveIntermediateTextures, CurrentMaterialSettings, &StagingBufferPool, &Output, Property, MaterialIndex, bEmissiveHDR = bEmissiveHDR](FTexture2DRHIRef& StagingBuffer, void * Data, int32 DataWidth, int32 DataHeight)
-								{
-									TRACE_CPUPROFILER_EVENT_SCOPE(FinalProcessing)
-
-									FBakeOutputEx& CurrentOutput  = Output[MaterialIndex];
-									TArray<FColor>& OutputColor = CurrentOutput.PropertyData[Property];
-									FIntPoint& OutputSize       = CurrentOutput.PropertySizes[Property];
-
-									OutputColor.SetNum(OutputSize.X * OutputSize.Y);
-
-									if (Property.Type == MP_EmissiveColor)
-									{
-										// Only one thread will write to CurrentOutput.EmissiveScale since there can be only one emissive channel property per FBakeOutputEx
-										FMaterialBakingModule::ProcessEmissiveOutput((const FFloat16Color*)Data, DataWidth, OutputSize, OutputColor, CurrentOutput.EmissiveScale, CurrentMaterialSettings->BackgroundColor);
-
-										if (bEmissiveHDR)
-										{
-											TArray<FFloat16Color>& OutputHDRColor = CurrentOutput.HDRPropertyData[Property];
-											OutputHDRColor.SetNum(OutputSize.X * OutputSize.Y);
-											ConvertRawR16G16B16A16FDataToFFloat16Color(OutputSize.X, OutputSize.Y, (uint8*)Data, DataWidth * sizeof(FFloat16Color), OutputHDRColor.GetData());
-										}
-									}
-									else
-									{
-										TRACE_CPUPROFILER_EVENT_SCOPE(ConvertRawB8G8R8A8DataToFColor)
-										
-										check(StagingBuffer->GetFormat() == PF_B8G8R8A8);
-										ConvertRawB8G8R8A8DataToFColor(OutputSize.X, OutputSize.Y, (uint8*)Data, DataWidth * sizeof(FColor), OutputColor.GetData());
-									}
-
-									// We can't unmap ourself since we're not on the render thread
-									StagingBufferPool.ReleaseStagingBufferForUnmap_AnyThread(StagingBuffer);
-
-									if (CurrentMaterialSettings->bPerformShrinking)
-									{
-										FMaterialBakingHelpers::PerformShrinking(OutputColor, OutputSize.X, OutputSize.Y, CurrentMaterialSettings->BackgroundColor);
-									}
-
-									if (CurrentMaterialSettings->bPerformBorderSmear)
-									{
-										FMaterialBakingHelpers::PerformUVBorderSmear(OutputColor, OutputSize.X, OutputSize.Y, -1, CurrentMaterialSettings->BackgroundColor);
-									}
-#if WITH_EDITOR
-									// If saving intermediates is turned on
-									if (bSaveIntermediateTextures)
-									{
-										TRACE_CPUPROFILER_EVENT_SCOPE(SaveIntermediateTextures)
-										FString TrimmedPropertyName = Property.ToString();
-
-										const FString DirectoryPath = FPaths::ConvertRelativePathToFull(FPaths::ProjectIntermediateDir() + TEXT("MaterialBaking/"));
-										FString FilenameString = FString::Printf(TEXT("%s%s-%d-%s.bmp"), *DirectoryPath, *CurrentMaterialSettings->Material->GetName(), MaterialIndex, *TrimmedPropertyName);
-										FFileHelper::CreateBitmap(*FilenameString, CurrentOutput.PropertySizes[Property].X, CurrentOutput.PropertySizes[Property].Y, CurrentOutput.PropertyData[Property].GetData());
-									}
-#endif // WITH_EDITOR
-									NumTasks--;
-								};
-
-							// Run previous command if we're going to overwrite it meaning pipeline depth has been reached
-							if (PipelineContext[PipelineIndex].ReadCommand)
-							{
-								PipelineContext[PipelineIndex].ReadCommand(RHICmdList);
-							}
+				SaveIntermediateTextures(Output, Property, TEXT("Final"));
+			}
+		}
+	}
 
-							// Generate a texture reading command that will be executed once it reaches the end of the pipeline
-							PipelineContext[PipelineIndex].ReadCommand =
-								[FinalProcessing_AnyThread, StagingBufferRef = MoveTemp(StagingBufferRef), GPUFence = MoveTemp(GPUFence)](FRHICommandListImmediate& RHICmdList) mutable
-								{
-									TRACE_CPUPROFILER_EVENT_SCOPE(MapAndEnqueue)
-
-									void * Data = nullptr;
-									int32 Width; int32 Height;
-									RHICmdList.MapStagingSurface(StagingBufferRef, GPUFence.GetReference(), Data, Width, Height);
-
-									// Schedule the copy and processing on another thread to free up the render thread as much as possible
-									Async(
-										EAsyncExecution::ThreadPool,
-										[FinalProcessing_AnyThread, Data, Width, Height, StagingBufferRef = MoveTemp(StagingBufferRef)]() mutable
-										{
-											FinalProcessing_AnyThread(StagingBufferRef, Data, Width, Height);
-										}
-									);
-								};
-						}
-					);
+	virtual FBakeOutputEx& GetBakeOutput(int32 InMaterialIndex) override
+	{
+		return Output;
+	}
 
-					PipelineIndex = (PipelineIndex + 1) % PipelineDepth;
-				}
-			}
+	virtual UTextureRenderTarget2D* GetRenderTarget(FMaterialPropertyEx InMaterialProperty, const FIntPoint& InRequiredSize, const FMaterialDataEx& InMaterialSettings) override
+	{
+		return RenderTargets[InMaterialProperty];
+	}
 
-		}
+private:
+	TMap<FMaterialPropertyEx, UTextureRenderTarget2D*> RenderTargets;
+	FBakeOutputEx& Output;
+};
 
-		// Destroying Render Items must happen on the render thread to ensure
-		// they are not used anymore.
-		ENQUEUE_RENDER_COMMAND(DestroyRenderItems)(
-			[RenderItems](FRHICommandListImmediate& RHICmdList)
+
+class FMaterialBakingProcessorMultiOutput : public FMaterialBakingProcessor
+{
+public:
+	FMaterialBakingProcessorMultiOutput(FMaterialBakingModule& InMaterialBakingModule, const TArray<FMaterialDataEx*>& InMaterialSettings, const TArray<FMeshData*>& InMeshSettings, TArray<FBakeOutputEx>& InOutput)
+		: FMaterialBakingProcessor(InMaterialBakingModule, InMaterialSettings, InMeshSettings)
+		, Output(InOutput)
+	{
+		Output.SetNum(InMaterialSettings.Num());
+
+		int32 NumOutputs = InOutput.Num();
+		if (NumOutputs != 0)
+		{
+			if (ensure(NumOutputs == InMaterialSettings.Num()))
 			{
-				for (auto RenderItem : (*RenderItems))
+				for (int32 Idx = 0; Idx < NumOutputs; ++Idx)
 				{
-					delete RenderItem.Value;
+					PrepareBakeOutput(InMaterialSettings[Idx], Output[Idx]);
 				}
-
-				delete RenderItems;
 			}
-		);
+		}
 	}
-
-	ENQUEUE_RENDER_COMMAND(ProcessRemainingReads)(
-		[&PipelineContext, PipelineDepth, PipelineIndex](FRHICommandListImmediate& RHICmdList)
+	
+	~FMaterialBakingProcessorMultiOutput()
+	{
+		ENQUEUE_RENDER_COMMAND(ProcessRemainingReads)(
+		[this, PipelineIndex=PipelineIndex](FRHICommandListImmediate& RHICmdList)
 		{
 			// Enqueue remaining reads
 			for (int32 Index = 0; Index < PipelineDepth; Index++)
@@ -810,35 +900,172 @@ void FMaterialBakingModule::BakeMaterials(const TArray<FMaterialDataEx*>& Materi
 					PipelineContext[LocalPipelineIndex].ReadCommand(RHICmdList);
 				}
 			}
+		});
+
+		// Wait until every tasks have been queued so that NumTasks is only decreasing
+		FlushRenderingCommands();
+
+		// Wait for any remaining final processing tasks
+		while (NumTasks.Load(EMemoryOrder::Relaxed) > 0)
+		{
+			FPlatformProcess::Sleep(0.1f);
 		}
-	);
 
-	// Wait until every tasks have been queued so that NumTasks is only decreasing
-	FlushRenderingCommands();
+		// Wait for all tasks to have been processed before clearing the staging buffers
+		FlushRenderingCommands();
 
-	// Wait for any remaining final processing tasks
-	while (NumTasks.Load(EMemoryOrder::Relaxed) > 0)
-	{
-		FPlatformProcess::Sleep(0.1f);
+		ENQUEUE_RENDER_COMMAND(ClearStagingBufferPool)(
+			[this](FRHICommandListImmediate& RHICmdList)
+			{
+				StagingBufferPool.Clear_RenderThread(RHICmdList);
+			}
+		);
+
+		// Wait for StagingBufferPool clear to have executed before exiting the function
+		FlushRenderingCommands();
 	}
 
-	// Wait for all tasks to have been processed before clearing the staging buffers
-	FlushRenderingCommands();
+private:
+	virtual void OnMaterialPropertyBaked(const FMaterialDataEx& CurrentMaterialSettings, const FMaterialPropertyEx& Property, FMeshMaterialRenderItem* RenderItem, UTextureRenderTarget2D* RenderTarget, FExportMaterialProxy* ExportMaterialProxy, FBakeOutputEx& CurrentOutput) override
+	{
+		ENQUEUE_RENDER_COMMAND(CopyStagingBuffer)(
+			[this, RenderTarget, Property, ExportMaterialProxy, &CurrentMaterialSettings, &CurrentOutput, PipelineIndex=PipelineIndex](FRHICommandListImmediate& RHICmdList)
+			{
+				FTextureRenderTargetResource* RenderTargetResource = RenderTarget->GetRenderTargetResource();
 
-	ENQUEUE_RENDER_COMMAND(ClearStagingBufferPool)(
-		[&StagingBufferPool](FRHICommandListImmediate& RHICmdList)
-		{
-			StagingBufferPool.Clear_RenderThread(RHICmdList);
-		}
-	);
+				FTexture2DRHIRef StagingBufferRef = StagingBufferPool.CreateStagingBuffer_RenderThread(RHICmdList, RenderTargetResource->GetSizeX(), RenderTargetResource->GetSizeY(), RenderTarget->GetFormat(), RenderTarget->IsSRGB());
+				FGPUFenceRHIRef GPUFence = RHICreateGPUFence(TEXT("MaterialBackingFence"));
+				TransitionAndCopyTexture(RHICmdList, RenderTargetResource->GetRenderTargetTexture(), StagingBufferRef, {});
+				RHICmdList.WriteGPUFence(GPUFence);
+
+				// Prepare a lambda for final processing that will be executed asynchronously
+				NumTasks++;
+				auto FinalProcessing_AnyThread =
+					[this, CurrentMaterialSettings, &CurrentOutput, Property](FTexture2DRHIRef& StagingBuffer, void* Data, int32 DataWidth, int32 DataHeight)
+					{
+						TRACE_CPUPROFILER_EVENT_SCOPE(FinalProcessing)
+
+						TArray<FColor>& OutputColor = CurrentOutput.PropertyData[Property];
+						FIntPoint& OutputSize = CurrentOutput.PropertySizes[Property];
+
+						OutputColor.SetNum(OutputSize.X * OutputSize.Y);
 
-	// Wait for StagingBufferPool clear to have executed before exiting the function
-	FlushRenderingCommands();
+						if (Property.Type == MP_EmissiveColor)
+						{
+							// Only one thread will write to CurrentOutput.EmissiveScale since there can be only one emissive channel property per FBakeOutputEx
+							ProcessEmissiveOutput((const FFloat16Color*)Data, DataWidth, OutputSize, OutputColor, CurrentOutput.EmissiveScale, CurrentMaterialSettings.BackgroundColor);
+
+							if (bEmissiveHDR)
+							{
+								TArray<FFloat16Color>& OutputHDRColor = CurrentOutput.HDRPropertyData[Property];
+								OutputHDRColor.SetNum(OutputSize.X * OutputSize.Y);
+								ConvertRawR16G16B16A16FDataToFFloat16Color(OutputSize.X, OutputSize.Y, (uint8*)Data, DataWidth * sizeof(FFloat16Color), OutputHDRColor.GetData());
+							}
+						}
+						else
+						{
+							TRACE_CPUPROFILER_EVENT_SCOPE(ConvertRawB8G8R8A8DataToFColor)
 
-	if (!CVarUseMaterialProxyCaching.GetValueOnAnyThread())
+							check(StagingBuffer->GetFormat() == PF_B8G8R8A8);
+							ConvertRawB8G8R8A8DataToFColor(OutputSize.X, OutputSize.Y, (uint8*)Data, DataWidth * sizeof(FColor), OutputColor.GetData());
+						}
+
+						// We can't unmap ourself since we're not on the render thread
+						StagingBufferPool.ReleaseStagingBufferForUnmap_AnyThread(StagingBuffer);
+
+						if (CurrentMaterialSettings.bPerformShrinking)
+						{
+							FMaterialBakingHelpers::PerformShrinking(OutputColor, OutputSize.X, OutputSize.Y, CurrentMaterialSettings.BackgroundColor);
+						}
+
+						if (CurrentMaterialSettings.bPerformBorderSmear)
+						{
+							// This will resize the output to a single pixel if the result is monochrome.
+							FMaterialBakingHelpers::PerformUVBorderSmearAndShrink(OutputColor, OutputSize.X, OutputSize.Y, CurrentMaterialSettings.BackgroundColor);
+						}
+	
+						SaveIntermediateTextures(CurrentOutput, Property, *CurrentMaterialSettings.Material->GetName());
+	
+						NumTasks--;
+					};
+
+				// Run previous command if we're going to overwrite it meaning pipeline depth has been reached
+				if (PipelineContext[PipelineIndex].ReadCommand)
+				{
+					PipelineContext[PipelineIndex].ReadCommand(RHICmdList);
+				}
+
+				// Generate a texture reading command that will be executed once it reaches the end of the pipeline
+				PipelineContext[PipelineIndex].ReadCommand =
+					[FinalProcessing_AnyThread, StagingBufferRef = MoveTemp(StagingBufferRef), GPUFence = MoveTemp(GPUFence)](FRHICommandListImmediate& RHICmdList) mutable
+					{
+						TRACE_CPUPROFILER_EVENT_SCOPE(MapAndEnqueue)
+
+						void* Data = nullptr;
+						int32 Width; int32 Height;
+						RHICmdList.MapStagingSurface(StagingBufferRef, GPUFence.GetReference(), Data, Width, Height);
+
+						// Schedule the copy and processing on another thread to free up the render thread as much as possible
+						Async(
+							EAsyncExecution::ThreadPool,
+							[FinalProcessing_AnyThread, Data, Width, Height, StagingBufferRef = MoveTemp(StagingBufferRef)]() mutable
+							{
+								FinalProcessing_AnyThread(StagingBufferRef, Data, Width, Height);
+							}
+						);
+					};
+			}
+		);
+
+		PipelineIndex = (PipelineIndex + 1) % PipelineDepth;
+	}
+
+	virtual FBakeOutputEx& GetBakeOutput(int32 InMaterialIndex) override
+	{
+		return Output[InMaterialIndex];
+	}
+
+	virtual UTextureRenderTarget2D* GetRenderTarget(FMaterialPropertyEx InMaterialProperty, const FIntPoint& InRequiredSize, const FMaterialDataEx& InMaterialSettings) override
 	{
-		CleanupMaterialProxies();
+		// It is safe to reuse the same render targets for each draw pass since they all execute sequentially on the GPU and are copied to staging buffers before
+		// being reused.
+		const bool bUsePooledRenderTargets = true;
+		return CreateRenderTarget(InMaterialProperty, InRequiredSize, bUsePooledRenderTargets, InMaterialSettings.BackgroundColor);
 	}
+
+private:
+	struct FPipelineContext
+	{
+		typedef TFunction<void(FRHICommandListImmediate& RHICmdList)> FReadCommand;
+		FReadCommand ReadCommand;
+	};
+
+	// Distance between the command sent to rendering and the GPU read-back of the result
+	// to minimize sync time waiting on GPU.
+	int32 PipelineIndex = 0;
+	FPipelineContext PipelineContext[PipelineDepth];
+
+	TAtomic<uint32>    NumTasks = 0;
+	FStagingBufferPool StagingBufferPool;
+
+	TArray<FBakeOutputEx>& Output;
+};
+
+void FMaterialBakingModule::BakeMaterials(const TArray<FMaterialDataEx*>& MaterialSettings, const TArray<FMeshData*>& MeshSettings, FBakeOutputEx& Output)
+{
+	TRACE_CPUPROFILER_EVENT_SCOPE(FMaterialBakingModule::BakeMaterials)
+
+	FMaterialBakingProcessorSingleOutput MaterialBakingProcessor(*this, MaterialSettings, MeshSettings, Output);
+	MaterialBakingProcessor.BakeMaterials();
+}
+
+
+void FMaterialBakingModule::BakeMaterials(const TArray<FMaterialDataEx*>& MaterialSettings, const TArray<FMeshData*>& MeshSettings, TArray<FBakeOutputEx>& Output)
+{
+	TRACE_CPUPROFILER_EVENT_SCOPE(FMaterialBakingModule::BakeMaterials)
+
+	FMaterialBakingProcessorMultiOutput MaterialBakingProcessor(*this, MaterialSettings, MeshSettings, Output);
+	MaterialBakingProcessor.BakeMaterials();
 }
 
 bool FMaterialBakingModule::SetupMaterialBakeSettings(TArray<TWeakObjectPtr<UObject>>& OptionObjects, int32 NumLODs)
@@ -923,40 +1150,73 @@ void FMaterialBakingModule::CleanupMaterialProxies()
 	MaterialProxyPool.Reset();
 }
 
-UTextureRenderTarget2D* FMaterialBakingModule::CreateRenderTarget(bool bInForceLinearGamma, EPixelFormat InPixelFormat, const FIntPoint& InTargetSize, const FColor& BackgroundColor)
+void FMaterialBakingModule::CleanupRenderTargets()
+{
+	for (auto RenderTarget : RenderTargetPool)
+	{
+		RenderTarget->RemoveFromRoot();
+	}
+	RenderTargetPool.Empty();
+}
+
+UTextureRenderTarget2D* FMaterialBakingModule::CreateRenderTarget(FMaterialPropertyEx InProperty, const FIntPoint& InTargetSize, bool bInUsePooledRenderTargets, const FColor& BackgroundColor)
 {
 	TRACE_CPUPROFILER_EVENT_SCOPE(FMaterialBakingModule::CreateRenderTarget)
 
-	UTextureRenderTarget2D* RenderTarget = nullptr;
+	// Lookup gamma and format settings for property, if not found use default values
+
+	// Gamma
+	const EPropertyColorSpace* OverrideColorSpace = PerPropertyColorSpace.Find(InProperty);
+	const EPropertyColorSpace ColorSpace = OverrideColorSpace ? *OverrideColorSpace : DefaultColorSpace;
+	const bool bForceLinearGamma = ColorSpace == EPropertyColorSpace::Linear;
+
+	// Pixel format
+	const EPixelFormat PixelFormat = PerPropertyFormat.Contains(InProperty) ? PerPropertyFormat[InProperty] : PF_B8G8R8A8;
+
 	const int32 MaxTextureSize = 1 << (MAX_TEXTURE_MIP_COUNT - 1); // Don't use GetMax2DTextureDimension() as this is for the RHI only.
 	const FIntPoint ClampedTargetSize(FMath::Clamp(InTargetSize.X, 1, MaxTextureSize), FMath::Clamp(InTargetSize.Y, 1, MaxTextureSize));
-	auto RenderTargetComparison = [bInForceLinearGamma, InPixelFormat, ClampedTargetSize](const UTextureRenderTarget2D* CompareRenderTarget) -> bool
-	{
-		return (CompareRenderTarget->SizeX == ClampedTargetSize.X && CompareRenderTarget->SizeY == ClampedTargetSize.Y && CompareRenderTarget->OverrideFormat == InPixelFormat && CompareRenderTarget->bForceLinearGamma == bInForceLinearGamma);
-	};
 
-	// Find any pooled render target with suitable properties.
-	UTextureRenderTarget2D** FindResult = RenderTargetPool.FindByPredicate(RenderTargetComparison);
-	
-	if (FindResult)
+	UTextureRenderTarget2D* RenderTarget = nullptr;
+
+	// First, look in pool
+	if (bInUsePooledRenderTargets)
 	{
-		RenderTarget = *FindResult;
+		auto RenderTargetComparison = [bForceLinearGamma, PixelFormat, ClampedTargetSize](const UTextureRenderTarget2D* CompareRenderTarget) -> bool
+		{
+			return (CompareRenderTarget->SizeX == ClampedTargetSize.X && CompareRenderTarget->SizeY == ClampedTargetSize.Y && CompareRenderTarget->OverrideFormat == PixelFormat && CompareRenderTarget->bForceLinearGamma == bForceLinearGamma);
+		};
+
+		// Find any pooled render target with suitable properties.
+		UTextureRenderTarget2D** FindResult = RenderTargetPool.FindByPredicate(RenderTargetComparison);
+		if (FindResult)
+		{
+			RenderTarget = *FindResult;
+		}
 	}
-	else
+
+	// If we want to avoid pooling, or no render target was found in the pool, create a new one
+	if (!RenderTarget)
 	{
 		TRACE_CPUPROFILER_EVENT_SCOPE(CreateNewRenderTarget)
 
 		// Not found - create a new one.
 		RenderTarget = NewObject<UTextureRenderTarget2D>();
 		check(RenderTarget);
-		RenderTarget->AddToRoot();
-		RenderTarget->ClearColor = bInForceLinearGamma ? BackgroundColor.ReinterpretAsLinear() : FLinearColor(BackgroundColor);
+
+		RenderTarget->ClearColor = bForceLinearGamma ? BackgroundColor.ReinterpretAsLinear() : FLinearColor(BackgroundColor);
 		RenderTarget->TargetGamma = 0.0f;
-		RenderTarget->InitCustomFormat(ClampedTargetSize.X, ClampedTargetSize.Y, InPixelFormat, bInForceLinearGamma);
+		RenderTarget->InitCustomFormat(ClampedTargetSize.X, ClampedTargetSize.Y, PixelFormat, bForceLinearGamma);
 
-		RenderTargetPool.Add(RenderTarget);
+		if (bInUsePooledRenderTargets)
+		{
+			RenderTarget->AddToRoot();
+			RenderTargetPool.Add(RenderTarget);
+		}
 	}
 
+	const bool bClearRenderTarget = true;
+	RenderTarget->UpdateResourceImmediate(bClearRenderTarget);
+
 	checkf(RenderTarget != nullptr, TEXT("Unable to create or find valid render target"));
 	return RenderTarget;
 }
@@ -991,7 +1251,7 @@ FExportMaterialProxy* FMaterialBakingModule::CreateMaterialProxy(const FMaterial
 	return Proxy;
 }
 
-void FMaterialBakingModule::ProcessEmissiveOutput(const FFloat16Color* Color16, int32 Color16Pitch, const FIntPoint& OutputSize, TArray<FColor>& OutputColor, float& EmissiveScale, const FColor& BackgroundColor)
+void ProcessEmissiveOutput(const FFloat16Color* Color16, int32 Color16Pitch, const FIntPoint& OutputSize, TArray<FColor>& OutputColor, float& EmissiveScale, const FColor& BackgroundColor)
 {
 	TRACE_CPUPROFILER_EVENT_SCOPE(FMaterialBakingModule::ProcessEmissiveOutput)
 
diff --git a/Engine/Source/Developer/MaterialBaking/Private/MaterialBakingModule.h b/Engine/Source/Developer/MaterialBaking/Private/MaterialBakingModule.h
--- a/Engine/Source/Developer/MaterialBaking/Private/MaterialBakingModule.h
+++ b/Engine/Source/Developer/MaterialBaking/Private/MaterialBakingModule.h
@@ -34,6 +34,12 @@ public:
 	/** Bakes out material properties according to extended MaterialSettings using MeshSettings and stores the output in Output */
 	virtual void BakeMaterials(const TArray<FMaterialDataEx*>& MaterialSettings, const TArray<FMeshData*>& MeshSettings, TArray<FBakeOutputEx>& Output) override;
 
+	/** Bakes out material properties according to MaterialSettings using MeshSettings and stores the output in Output */
+	virtual void BakeMaterials(const TArray<FMaterialData*>& MaterialSettings, const TArray<FMeshData*>& MeshSettings, FBakeOutput& Output) override;
+
+	/** Bakes out material properties according to extended MaterialSettings using MeshSettings and stores the output in Output */
+	virtual void BakeMaterials(const TArray<FMaterialDataEx*>& MaterialSettings, const TArray<FMeshData*>& MeshSettings, FBakeOutputEx& Output) override;
+
 	/** Promps a slate window to allow the user to populate specific material baking settings used while baking out materials */
 	virtual bool SetupMaterialBakeSettings(TArray<TWeakObjectPtr<UObject>>& OptionObjects, int32 NumLODs) override;
 
@@ -50,18 +56,20 @@ public:
 	virtual uint32 GetCRC() const override;
 
 protected:
+	friend class FMaterialBakingProcessor;
+
 	/* Creates and adds or reuses a RenderTarget from the pool */
-	UTextureRenderTarget2D* CreateRenderTarget(bool bInForceLinearGamma, EPixelFormat InPixelFormat, const FIntPoint& InTargetSize, const FColor& BackgroundColor);
+	UTextureRenderTarget2D* CreateRenderTarget(FMaterialPropertyEx InProperty, const FIntPoint& InTargetSize, bool bInUsePooledRenderTargets, const FColor& BackgroundColor);
 
 	/* Creates and adds (or reuses a ExportMaterialProxy from the pool if MaterialBaking.UseMaterialProxyCaching is set to 1) */
 	FExportMaterialProxy* CreateMaterialProxy(const FMaterialDataEx* MaterialSettings, const FMaterialPropertyEx& Property);
 
-	/** Helper for emissive color conversion to Output */
-	static void ProcessEmissiveOutput(const FFloat16Color* Color16, int32 Color16Pitch, const FIntPoint& OutputSize, TArray<FColor>& Output, float& EmissiveScale, const FColor& BackgroundColor);
-
 	/** Cleans up all cached material proxies in MaterialProxyPool */
 	void CleanupMaterialProxies();
 
+	/** Free up all pooled render targets */
+	void CleanupRenderTargets();
+
 	/** Callback for modified objects which should be removed from MaterialProxyPool in that case */
 	void OnObjectModified(UObject* Object);
 
diff --git a/Engine/Source/Developer/MaterialBaking/Public/IMaterialBakingModule.h b/Engine/Source/Developer/MaterialBaking/Public/IMaterialBakingModule.h
--- a/Engine/Source/Developer/MaterialBaking/Public/IMaterialBakingModule.h
+++ b/Engine/Source/Developer/MaterialBaking/Public/IMaterialBakingModule.h
@@ -30,6 +30,12 @@ public:
 	/** Bakes out material properties according to extended MaterialSettings using MeshSettings and stores the output in Output */
 	virtual void BakeMaterials(const TArray<FMaterialDataEx*>& MaterialSettings, const TArray<FMeshData*>& MeshSettings, TArray<FBakeOutputEx>& Output) = 0;
 
+	/** Bakes out material properties according to MaterialSettings using MeshSettings and stores the output in Output */
+	virtual void BakeMaterials(const TArray<FMaterialData*>& MaterialSettings, const TArray<FMeshData*>& MeshSettings, FBakeOutput& Output) = 0;
+
+	/** Bakes out material properties according to extended MaterialSettings using MeshSettings and stores the output in Output */
+	virtual void BakeMaterials(const TArray<FMaterialDataEx*>& MaterialSettings, const TArray<FMeshData*>& MeshSettings, FBakeOutputEx& Output) = 0;
+
 	/** Promps a slate window to allow the user to populate specific material baking settings used while baking out materials */
 	virtual  bool SetupMaterialBakeSettings(TArray<TWeakObjectPtr<UObject>>& OptionObjects, int32 NumLODs) = 0;
 
diff --git a/Engine/Source/Developer/MeshMergeUtilities/Private/MeshMergeUtilities.cpp b/Engine/Source/Developer/MeshMergeUtilities/Private/MeshMergeUtilities.cpp
--- a/Engine/Source/Developer/MeshMergeUtilities/Private/MeshMergeUtilities.cpp
+++ b/Engine/Source/Developer/MeshMergeUtilities/Private/MeshMergeUtilities.cpp
@@ -850,17 +850,8 @@ FIntPoint FMeshMergeUtilities::ConditionalImageResize(const FIntPoint& SrcSize,
 
 void FMeshMergeUtilities::MergeFlattenedMaterials(TArray<struct FFlattenMaterial>& InMaterialList, int32 InGutter, FFlattenMaterial& OutMergedMaterial, TArray<FUVOffsetScalePair>& OutUVTransforms) const
 {
-	OutUVTransforms.Reserve(InMaterialList.Num());
-
 	// Fill output UV transforms with invalid values
-	for (auto& Material : InMaterialList)
-	{
-		// Invalid UV transform
-		FUVOffsetScalePair UVTransform;
-		UVTransform.Key = FVector2D::ZeroVector;
-		UVTransform.Value = FVector2D::ZeroVector;
-		OutUVTransforms.Add(UVTransform);
-	}
+	OutUVTransforms.SetNumZeroed(InMaterialList.Num());
 
 	const int32 AtlasGridSize = FMath::CeilToInt(FMath::Sqrt(static_cast<float>(InMaterialList.Num())));
 	OutMergedMaterial.EmissiveScale = FlattenEmissivescale(InMaterialList);
@@ -979,7 +970,9 @@ void FMeshMergeUtilities::MergeFlattenedMaterials(TArray<struct FFlattenMaterial
 
 void FMeshMergeUtilities::FlattenBinnedMaterials(TArray<struct FFlattenMaterial>& InMaterialList, const TArray<FBox2D>& InMaterialBoxes, int32 InGutter, bool bCopyOnlyMaskedPixels, FFlattenMaterial& OutMergedMaterial, TArray<FUVOffsetScalePair>& OutUVTransforms) const
 {
-	OutUVTransforms.AddZeroed(InMaterialList.Num());
+	// Fill output UV transforms with invalid values
+	OutUVTransforms.SetNumZeroed(InMaterialList.Num());
+
 	// Flatten emissive scale across all incoming materials
 	OutMergedMaterial.EmissiveScale = FlattenEmissivescale(InMaterialList);
 
@@ -2962,7 +2955,18 @@ void FMeshMergeUtilities::CreateMergedMaterial(FMeshMergeDataTracker& InDataTrac
 	{
 		TArray<FBakeOutput> BakeOutputs;
 		IMaterialBakingModule& Module = FModuleManager::Get().LoadModuleChecked<IMaterialBakingModule>("MaterialBaking");
-		Module.BakeMaterials(MaterialSettingPtrs, MeshSettingPtrs, BakeOutputs);
+
+		// If we're working with a new set of UVs, we can bake all materials directly to the same bake output
+		// as our remapped UVs for each mesh don't overlap.
+		if (bGloballyRemapUVs)
+		{
+			FBakeOutput& BakeOutput = BakeOutputs.Emplace_GetRef();
+			Module.BakeMaterials(MaterialSettingPtrs, MeshSettingPtrs, BakeOutput);
+		}
+		else
+		{
+			Module.BakeMaterials(MaterialSettingPtrs, MeshSettingPtrs, BakeOutputs);
+		}
 
 		// Append constant properties ?
 		TArray<FColor> ConstantData;
@@ -3018,6 +3022,9 @@ void FMeshMergeUtilities::CreateMergedMaterial(FMeshMergeDataTracker& InDataTrac
 		}
 
 		FlattenBinnedMaterials(FlattenedMaterials, MaterialBoxes, 0, true, OutFlattenMaterial, UVTransforms);
+
+		static const FUVOffsetScalePair NoUVTransform = { FVector2D::Zero(), FVector2D::One() };
+		UVTransforms.Init(NoUVTransform, GlobalMaterialSettings.Num());
 	}
 	else
 	{
