Material Baking - Fix incorrect use of ProcessEmissiveOutput()
* Pitch in is pixels, not bytes
#rnx

diff --git a/Engine/Source/Developer/MaterialBaking/Private/MaterialBakingModule.cpp b/Engine/Source/Developer/MaterialBaking/Private/MaterialBakingModule.cpp
--- a/Engine/Source/Developer/MaterialBaking/Private/MaterialBakingModule.cpp
+++ b/Engine/Source/Developer/MaterialBaking/Private/MaterialBakingModule.cpp
@@ -827,8 +827,8 @@ public:
 					TArray<FFloat16Color> OutputHDRColor;
 					RTResource->ReadFloat16Pixels(OutputHDRColor);
 					
-					int32 Color16Pitch = Output.PropertySizes[Property].X * sizeof(FFloat16Color);
-					ProcessEmissiveOutput(OutputHDRColor.GetData(), Color16Pitch, Output.PropertySizes[Property], Output.PropertyData[Property], Output.EmissiveScale, DefaultMaterialData->BackgroundColor);
+					int32 Color16PitchPixels = Output.PropertySizes[Property].X;
+					ProcessEmissiveOutput(OutputHDRColor.GetData(), Color16PitchPixels, Output.PropertySizes[Property], Output.PropertyData[Property], Output.EmissiveScale, DefaultMaterialData->BackgroundColor);
 					
 					if (bEmissiveHDR)
 					{
@@ -1251,7 +1251,7 @@ FExportMaterialProxy* FMaterialBakingModule::CreateMaterialProxy(const FMaterial
 	return Proxy;
 }
 
-void ProcessEmissiveOutput(const FFloat16Color* Color16, int32 Color16Pitch, const FIntPoint& OutputSize, TArray<FColor>& OutputColor, float& EmissiveScale, const FColor& BackgroundColor)
+void ProcessEmissiveOutput(const FFloat16Color* Color16, int32 Color16PitchPixels, const FIntPoint& OutputSize, TArray<FColor>& OutputColor, float& EmissiveScale, const FColor& BackgroundColor)
 {
 	TRACE_CPUPROFILER_EVENT_SCOPE(FMaterialBakingModule::ProcessEmissiveOutput)
 
@@ -1270,13 +1270,13 @@ void ProcessEmissiveOutput(const FFloat16Color* Color16, int32 Color16Pitch, con
 		FMemory::Memset(MaxValue, 0, NumThreads * sizeof(MaxValue[0]));
 		
 		// Find maximum float value across texture
-		ParallelFor(NumThreads, [&Color16, LinesPerThread, MaxValue, OutputSize, Color16Pitch, BackgroundColor16](int32 Index)
+		ParallelFor(NumThreads, [&Color16, LinesPerThread, MaxValue, OutputSize, Color16PitchPixels, BackgroundColor16](int32 Index)
 		{
 			const int32 EndY = FMath::Min((Index + 1) * LinesPerThread, OutputSize.Y);			
 			float& CurrentMaxValue = MaxValue[Index];
 			for (int32 PixelY = Index * LinesPerThread; PixelY < EndY; ++PixelY)
 			{
-				const int32 SrcYOffset = PixelY * Color16Pitch;
+				const int32 SrcYOffset = PixelY * Color16PitchPixels;
 				for (int32 PixelX = 0; PixelX < OutputSize.X; PixelX++)
 				{
 					const FFloat16Color& Pixel16 = Color16[PixelX + SrcYOffset];
@@ -1310,12 +1310,12 @@ void ProcessEmissiveOutput(const FFloat16Color* Color16, int32 Color16Pitch, con
 	// Now convert Float16 to Color using the scale
 	OutputColor.SetNumUninitialized(OutputSize.X * OutputSize.Y);
 	const float Scale = 255.0f / GlobalMaxValue;
-	ParallelFor(NumThreads, [&Color16, LinesPerThread, &OutputColor, OutputSize, Color16Pitch, Scale, BackgroundColor16, BackgroundColor](int32 Index)
+	ParallelFor(NumThreads, [&Color16, LinesPerThread, &OutputColor, OutputSize, Color16PitchPixels, Scale, BackgroundColor16, BackgroundColor](int32 Index)
 	{
 		const int32 EndY = FMath::Min((Index + 1) * LinesPerThread, OutputSize.Y);
 		for (int32 PixelY = Index * LinesPerThread; PixelY < EndY; ++PixelY)
 		{
-			const int32 SrcYOffset = PixelY * Color16Pitch;
+			const int32 SrcYOffset = PixelY * Color16PitchPixels;
 			const int32 DstYOffset = PixelY * OutputSize.X;
 
 			for (int32 PixelX = 0; PixelX < OutputSize.X; PixelX++)
