Material Baking - Made baking more resilient to materials that output none of the requested properties (diffuse, normal, etc)

diff --git a/Engine/Source/Developer/MaterialBaking/Private/MaterialBakingModule.cpp b/Engine/Source/Developer/MaterialBaking/Private/MaterialBakingModule.cpp
--- a/Engine/Source/Developer/MaterialBaking/Private/MaterialBakingModule.cpp
+++ b/Engine/Source/Developer/MaterialBaking/Private/MaterialBakingModule.cpp
@@ -490,44 +490,48 @@ public:
 		for (int32 Index = 0; Index < NumMaterials; ++Index)
 		{
 			const int32 MaterialIndex = ProcessingOrder[Index];
-			const FMaterialDataEx& CurrentMaterialSettings = *MaterialSettings[MaterialIndex];
-			const FMeshData* CurrentMeshSettings = MeshSettings[MaterialIndex];
-			FBakeOutputEx& CurrentOutput = GetBakeOutput(MaterialIndex);
-
-			TMap<FRenderItemKey, FMeshMaterialRenderItem*>* RenderItems = GetRenderItems(Index);
-			check(RenderItems && !RenderItems->IsEmpty());
 
-			// For each property
-			for (const auto& [Property, Size] : CurrentMaterialSettings.PropertySizes)
+			const FMaterialDataEx& CurrentMaterialSettings = *MaterialSettings[MaterialIndex];
+			if (!CurrentMaterialSettings.PropertySizes.IsEmpty())
 			{
-				TRACE_CPUPROFILER_EVENT_SCOPE_TEXT(*Property.ToString())
+				const FMeshData* CurrentMeshSettings = MeshSettings[MaterialIndex];
+				FBakeOutputEx& CurrentOutput = GetBakeOutput(MaterialIndex);
 
-				FExportMaterialProxy* ExportMaterialProxy = MaterialBakingModule.CreateMaterialProxy(&CurrentMaterialSettings, Property);
-				if (!ExportMaterialProxy->IsCompilationFinished())
-				{
-					TRACE_CPUPROFILER_EVENT_SCOPE(WaitForMaterialProxyCompilation)
-					ExportMaterialProxy->FinishCompilation();
-				}
+				TMap<FRenderItemKey, FMeshMaterialRenderItem*>* RenderItems = GetRenderItems(Index);
+				check(RenderItems && !RenderItems->IsEmpty());
 
-				UTextureRenderTarget2D* RenderTarget = GetRenderTarget(Property, Size, CurrentMaterialSettings);
-				FMeshMaterialRenderItem* RenderItem = RenderItems->FindChecked(FRenderItemKey(CurrentMeshSettings, Size));
-
-				BakeMaterialProperty(CurrentMaterialSettings, Property, RenderItem, RenderTarget, ExportMaterialProxy, CurrentOutput);
-			}
-
-			// Destroying Render Items
-			// Must happen on the render thread to ensure they are not used anymore.
-			ENQUEUE_RENDER_COMMAND(DestroyRenderItems)(
-				[RenderItems](FRHICommandListImmediate& RHICmdList)
+				// For each property
+				for (const auto& [Property, Size] : CurrentMaterialSettings.PropertySizes)
 				{
-					for (auto RenderItem : (*RenderItems))
+					TRACE_CPUPROFILER_EVENT_SCOPE_TEXT(*Property.ToString())
+
+					FExportMaterialProxy* ExportMaterialProxy = MaterialBakingModule.CreateMaterialProxy(&CurrentMaterialSettings, Property);
+					if (!ExportMaterialProxy->IsCompilationFinished())
 					{
-						delete RenderItem.Value;
+						TRACE_CPUPROFILER_EVENT_SCOPE(WaitForMaterialProxyCompilation)
+						ExportMaterialProxy->FinishCompilation();
 					}
 
-					delete RenderItems;
+					UTextureRenderTarget2D* RenderTarget = GetRenderTarget(Property, Size, CurrentMaterialSettings);
+					FMeshMaterialRenderItem* RenderItem = RenderItems->FindChecked(FRenderItemKey(CurrentMeshSettings, Size));
+
+					BakeMaterialProperty(CurrentMaterialSettings, Property, RenderItem, RenderTarget, ExportMaterialProxy, CurrentOutput);
 				}
-			);
+
+				// Destroying Render Items
+				// Must happen on the render thread to ensure they are not used anymore.
+				ENQUEUE_RENDER_COMMAND(DestroyRenderItems)(
+					[RenderItems](FRHICommandListImmediate& RHICmdList)
+					{
+						for (auto RenderItem : (*RenderItems))
+						{
+							delete RenderItem.Value;
+						}
+
+						delete RenderItems;
+					}
+				);
+			}
 		}
 	}
 
@@ -613,7 +617,6 @@ private:
 			}
 		}
 
-		check(!RenderItems->IsEmpty());
 		return RenderItems;
 	}
 
