Fix material baking incorrectly ignoring specific UV space when using mesh data

diff --git a/Engine/Source/Developer/MaterialBaking/Private/MaterialRenderItem.cpp b/Engine/Source/Developer/MaterialBaking/Private/MaterialRenderItem.cpp
--- a/Engine/Source/Developer/MaterialBaking/Private/MaterialRenderItem.cpp
+++ b/Engine/Source/Developer/MaterialBaking/Private/MaterialRenderItem.cpp
@@ -167,8 +167,8 @@ void FMeshMaterialRenderItem::PopulateWithQuadData()
 	Vertices.Empty(4);
 	Indices.Empty(6);
 
-	const float U = MeshSettings->TextureCoordinateBox.Min.X;
-	const float V = MeshSettings->TextureCoordinateBox.Min.Y;
+	const float OffsetU = MeshSettings->TextureCoordinateBox.Min.X;
+	const float OffsetV = MeshSettings->TextureCoordinateBox.Min.Y;
 	const float SizeU = MeshSettings->TextureCoordinateBox.Max.X - MeshSettings->TextureCoordinateBox.Min.X;
 	const float SizeV = MeshSettings->TextureCoordinateBox.Max.Y - MeshSettings->TextureCoordinateBox.Min.Y;
 	const float ScaleX = TextureSize.X;
@@ -188,7 +188,7 @@ void FMeshMaterialRenderItem::PopulateWithQuadData()
 		FMemory::Memzero(&Vert->TextureCoordinate, sizeof(Vert->TextureCoordinate));
 		for (int32 TexcoordIndex = 0; TexcoordIndex < MAX_STATIC_TEXCOORDS; TexcoordIndex++)
 		{
-			Vert->TextureCoordinate[TexcoordIndex].Set(U + SizeU * X, V + SizeV * Y);
+			Vert->TextureCoordinate[TexcoordIndex].Set(OffsetU + SizeU * X, OffsetV + SizeV * Y);
 		}
 		Vert->Color = FColor::White;
 	}
@@ -219,6 +219,10 @@ void FMeshMaterialRenderItem::PopulateWithMeshData()
 	Vertices.Empty(NumVerts);
 	Indices.Empty(NumVerts >> 1);
 
+	const float OffsetU = MeshSettings->TextureCoordinateBox.Min.X;
+	const float OffsetV = MeshSettings->TextureCoordinateBox.Min.Y;
+	const float SizeU = MeshSettings->TextureCoordinateBox.Max.X - MeshSettings->TextureCoordinateBox.Min.X;
+	const float SizeV = MeshSettings->TextureCoordinateBox.Max.Y - MeshSettings->TextureCoordinateBox.Min.Y;
 	const float ScaleX = TextureSize.X;
 	const float ScaleY = TextureSize.Y;
 
@@ -273,7 +277,8 @@ void FMeshMaterialRenderItem::PopulateWithMeshData()
 				Vert->SetTangents(TangentX, TangentY, TangentZ);
 				for (int32 TexcoordIndex = 0; TexcoordIndex < NumTexcoords; TexcoordIndex++)
 				{
-					Vert->TextureCoordinate[TexcoordIndex] = VertexInstanceUVs.Get(SrcVertexInstanceID, TexcoordIndex);
+					const FVector2f UV = VertexInstanceUVs.Get(SrcVertexInstanceID, TexcoordIndex);
+					Vert->TextureCoordinate[TexcoordIndex].Set(OffsetU + SizeU * UV.X, OffsetV + SizeV * UV.Y);
 				}
 
 				if (NumTexcoords < VertexPositionStoredUVChannel)
@@ -297,4 +302,4 @@ void FMeshMaterialRenderItem::PopulateWithMeshData()
 		}
 		FaceIndex++;
 	}
-}
\ No newline at end of file
+}
