Amend previous fix (25014983): Fix material baking incorrectly ignores specific UV space when using mesh data

Previous fix only works when mesh data is a quad, this amended fix works for arbitrary mesh data.

diff --git a/Engine/Source/Developer/MaterialBaking/Private/MaterialRenderItem.cpp b/Engine/Source/Developer/MaterialBaking/Private/MaterialRenderItem.cpp
--- a/Engine/Source/Developer/MaterialBaking/Private/MaterialRenderItem.cpp
+++ b/Engine/Source/Developer/MaterialBaking/Private/MaterialRenderItem.cpp
@@ -219,12 +219,11 @@ void FMeshMaterialRenderItem::PopulateWithMeshData()
 	Vertices.Empty(NumVerts);
 	Indices.Empty(NumVerts >> 1);
 
-	const float OffsetU = MeshSettings->TextureCoordinateBox.Min.X;
-	const float OffsetV = MeshSettings->TextureCoordinateBox.Min.Y;
-	const float SizeU = MeshSettings->TextureCoordinateBox.Max.X - MeshSettings->TextureCoordinateBox.Min.X;
-	const float SizeV = MeshSettings->TextureCoordinateBox.Max.Y - MeshSettings->TextureCoordinateBox.Min.Y;
-	const float ScaleX = TextureSize.X;
-	const float ScaleY = TextureSize.Y;
+	// When using arbitrary mesh data (rather than a simple quad), TextureCoordinateBox has to be applied to XY position:
+	const float ScaleX = TextureSize.X / (MeshSettings->TextureCoordinateBox.Max.X - MeshSettings->TextureCoordinateBox.Min.X);
+	const float ScaleY = TextureSize.Y / (MeshSettings->TextureCoordinateBox.Max.Y - MeshSettings->TextureCoordinateBox.Min.Y);
+	const float OffsetX = -MeshSettings->TextureCoordinateBox.Min.X * ScaleX;
+	const float OffsetY = -MeshSettings->TextureCoordinateBox.Min.Y * ScaleY;
 
 	const static int32 VertexPositionStoredUVChannel = 6;
 	// count number of texture coordinates for this mesh
@@ -264,12 +263,12 @@ void FMeshMaterialRenderItem::PopulateWithMeshData()
 				{
 					// compute vertex position from original UV
 					const FVector2D& UV = FVector2D(VertexInstanceUVs.Get(SrcVertexInstanceID, MeshSettings->TextureCoordinateIndex));
-					Vert->Position = WorldToLocal.TransformPosition(FVector3f(UV.X * ScaleX, UV.Y * ScaleY, 0));
+					Vert->Position = WorldToLocal.TransformPosition(FVector3f(OffsetX + UV.X * ScaleX, OffsetY + UV.Y * ScaleY, 0));
 				}
 				else
 				{
 					const FVector2D& UV = MeshSettings->CustomTextureCoordinates[SrcVertIndex];
-					Vert->Position = WorldToLocal.TransformPosition(FVector3f(UV.X * ScaleX, UV.Y * ScaleY, 0));
+					Vert->Position = WorldToLocal.TransformPosition(FVector3f(OffsetX + UV.X * ScaleX, OffsetY + UV.Y * ScaleY, 0));
 				}
 				FVector3f TangentX = WorldToLocal.TransformVector(VertexInstanceTangents[SrcVertexInstanceID]);
 				FVector3f TangentZ = WorldToLocal.TransformVector(VertexInstanceNormals[SrcVertexInstanceID]);
@@ -277,8 +276,7 @@ void FMeshMaterialRenderItem::PopulateWithMeshData()
 				Vert->SetTangents(TangentX, TangentY, TangentZ);
 				for (int32 TexcoordIndex = 0; TexcoordIndex < NumTexcoords; TexcoordIndex++)
 				{
-					const FVector2f UV = VertexInstanceUVs.Get(SrcVertexInstanceID, TexcoordIndex);
-					Vert->TextureCoordinate[TexcoordIndex].Set(OffsetU + SizeU * UV.X, OffsetV + SizeV * UV.Y);
+					Vert->TextureCoordinate[TexcoordIndex] = VertexInstanceUVs.Get(SrcVertexInstanceID, TexcoordIndex);
 				}
 
 				if (NumTexcoords < VertexPositionStoredUVChannel)
