Fix material baking incorrectly stating that EmissiveColor is baked as sRGB while in reality it's always baked as linear because of the HDR rendertarget.

Please note that this change doesn't actually affect any baking results, because EmissiveColor always reads linear data from its HDR rendertarget regardless of PerPropertyColorSpace. The fix only affects the value returned by FMaterialBakingModule::IsLinearBake for EmissiveColor, which until now has been incorrect due to the missmatch between PerPropertyColorSpace and how EmissiveColor is actually baked.

diff --git a/Engine/Source/Developer/MaterialBaking/Private/MaterialBakingModule.cpp b/Engine/Source/Developer/MaterialBaking/Private/MaterialBakingModule.cpp
--- a/Engine/Source/Developer/MaterialBaking/Private/MaterialBakingModule.cpp
+++ b/Engine/Source/Developer/MaterialBaking/Private/MaterialBakingModule.cpp
@@ -1112,13 +1112,13 @@ void FMaterialBakingModule::SetLinearBake(bool bCorrectLinear)
 	{
 		DefaultColorSpace = EPropertyColorSpace::Linear;
 		PerPropertyColorSpace.Add(MP_BaseColor, EPropertyColorSpace::sRGB);
-		PerPropertyColorSpace.Add(MP_EmissiveColor, EPropertyColorSpace::sRGB);
 		PerPropertyColorSpace.Add(MP_SubsurfaceColor, EPropertyColorSpace::sRGB);
 		PerPropertyColorSpace.Add(FMaterialPropertyEx::TransmittanceColor, EPropertyColorSpace::sRGB);
 	}
 	else
 	{
 		DefaultColorSpace = EPropertyColorSpace::sRGB;
+		PerPropertyColorSpace.Add(MP_EmissiveColor, EPropertyColorSpace::Linear); // Always linear because it uses HDR
 		PerPropertyColorSpace.Add(MP_Normal, EPropertyColorSpace::Linear);
 		PerPropertyColorSpace.Add(MP_Refraction, EPropertyColorSpace::Linear);
 		PerPropertyColorSpace.Add(MP_Opacity, EPropertyColorSpace::Linear);
diff --git a/Engine/Source/Developer/MaterialBaking/Private/MaterialBakingModule.h b/Engine/Source/Developer/MaterialBaking/Private/MaterialBakingModule.h
--- a/Engine/Source/Developer/MaterialBaking/Private/MaterialBakingModule.h
+++ b/Engine/Source/Developer/MaterialBaking/Private/MaterialBakingModule.h
@@ -46,7 +46,7 @@ public:
 	/** Outputs true HDR version of emissive color */
 	virtual void SetEmissiveHDR(bool bHDR) override;
 
-	/** Bakes all material properties to linear textures, except for colors */
+	/** Bakes all material properties to linear textures, except for non-emissive colors */
 	virtual void SetLinearBake(bool bCorrectLinear) override;
 
 	/** Returns whether a specific material property is baked to a linear texture or not */
