Material Baking - Readd support for MaterialBaking.RenderDocCapture cvar which wasn't doing anything anymore since a recent refactor

diff --git a/Engine/Source/Developer/MaterialBaking/Private/MaterialBakingModule.cpp b/Engine/Source/Developer/MaterialBaking/Private/MaterialBakingModule.cpp
--- a/Engine/Source/Developer/MaterialBaking/Private/MaterialBakingModule.cpp
+++ b/Engine/Source/Developer/MaterialBaking/Private/MaterialBakingModule.cpp
@@ -702,6 +702,8 @@ private:
 		ENQUEUE_RENDER_COMMAND(RenderOneMaterial)(
 			[RenderItem, RenderTarget, ExportMaterialProxy](FRHICommandListImmediate& RHICmdList)
 			{
+				RenderCaptureInterface::FScopedCapture RenderCapture(CVarMaterialBakingRDOCCapture.GetValueOnAnyThread() == 1, TEXT("MaterialBaking"));
+
 				FSceneViewFamily ViewFamily(FSceneViewFamily::ConstructionValues(RenderTarget->GetRenderTargetResource(), nullptr,
 					FEngineShowFlags(ESFIM_Game))
 					.SetTime(FGameTime())
