MaterialBaking: Use correct RenderCaptureInterface::FScopedCapture() constructor on the render thread.

diff --git a/Engine/Source/Developer/MaterialBaking/Private/MaterialBakingModule.cpp b/Engine/Source/Developer/MaterialBaking/Private/MaterialBakingModule.cpp
--- a/Engine/Source/Developer/MaterialBaking/Private/MaterialBakingModule.cpp
+++ b/Engine/Source/Developer/MaterialBaking/Private/MaterialBakingModule.cpp
@@ -702,7 +702,7 @@ private:
 		ENQUEUE_RENDER_COMMAND(RenderOneMaterial)(
 			[RenderItem, RenderTarget, ExportMaterialProxy](FRHICommandListImmediate& RHICmdList)
 			{
-				RenderCaptureInterface::FScopedCapture RenderCapture(CVarMaterialBakingRDOCCapture.GetValueOnAnyThread() == 1, TEXT("MaterialBaking"));
+				RenderCaptureInterface::FScopedCapture RenderCapture(CVarMaterialBakingRDOCCapture.GetValueOnAnyThread() == 1, &RHICmdList, TEXT("MaterialBaking"));
 
 				FSceneViewFamily ViewFamily(FSceneViewFamily::ConstructionValues(RenderTarget->GetRenderTargetResource(), nullptr,
 					FEngineShowFlags(ESFIM_Game))
