Fix crash in the Material Baking Module where it was trying to delete a material render proxy while shader compilation was happening in the background.

- We can't simply enqueue a render command and `delete` the proxy.
- Call the `FMaterial::DeferredDeleteArray` helper function which will properly cancel shader compilation if in-flight.

#rnx

diff --git a/Engine/Source/Developer/MaterialBaking/Private/MaterialBakingModule.cpp b/Engine/Source/Developer/MaterialBaking/Private/MaterialBakingModule.cpp
--- a/Engine/Source/Developer/MaterialBaking/Private/MaterialBakingModule.cpp
+++ b/Engine/Source/Developer/MaterialBaking/Private/MaterialBakingModule.cpp
@@ -1145,21 +1145,15 @@ bool FMaterialBakingModule::IsLinearBake(FMaterialPropertyEx Property)
 	return ColorSpace == EPropertyColorSpace::Linear;
 }
 
-static void DeleteCachedMaterialProxy(FExportMaterialProxy* Proxy)
-{
-	ENQUEUE_RENDER_COMMAND(DeleteCachedMaterialProxy)(
-		[Proxy](FRHICommandListImmediate& RHICmdList)
-		{
-			delete Proxy;
-		});
-}
-
 void FMaterialBakingModule::CleanupMaterialProxies()
 {
+	TArray<FMaterial*> ResourcesToFree;
 	for (auto Iterator : MaterialProxyPool)
 	{
-		DeleteCachedMaterialProxy(Iterator.Value.Value);
+		ResourcesToFree.Add(Iterator.Value.Value);
 	}
+	FMaterial::DeferredDeleteArray(ResourcesToFree);
+
 	MaterialProxyPool.Reset();
 }
 
@@ -1401,7 +1395,10 @@ void FMaterialBakingModule::OnObjectModified(UObject* Object)
 				// We have a match, remove the entry from our pool
 				if (bMustDelete)
 				{
-					DeleteCachedMaterialProxy(It.Value().Value);
+					TArray<FMaterial*> ResourcesToFree;
+					ResourcesToFree.Add(It.Value().Value);
+					FMaterial::DeferredDeleteArray(ResourcesToFree);
+
 					It.RemoveCurrent();
 				}
 			}
