Material Baking - Modified default logging verbosity to "Log" rather than "All"

diff --git a/Engine/Source/Developer/MaterialBaking/Private/MaterialBakingModule.cpp b/Engine/Source/Developer/MaterialBaking/Private/MaterialBakingModule.cpp
--- a/Engine/Source/Developer/MaterialBaking/Private/MaterialBakingModule.cpp
+++ b/Engine/Source/Developer/MaterialBaking/Private/MaterialBakingModule.cpp
@@ -34,7 +34,7 @@
 
 IMPLEMENT_MODULE(FMaterialBakingModule, MaterialBaking);
 
-DEFINE_LOG_CATEGORY_STATIC(LogMaterialBaking, All, All);
+DEFINE_LOG_CATEGORY_STATIC(LogMaterialBaking, Log, All);
 
 #define LOCTEXT_NAMESPACE "MaterialBakingModule"
 
