Material Baking - Better handle invalid entries
* Do not include empty entries in the array to be processed

diff --git a/Engine/Source/Developer/MaterialBaking/Private/MaterialBakingModule.cpp b/Engine/Source/Developer/MaterialBaking/Private/MaterialBakingModule.cpp
--- a/Engine/Source/Developer/MaterialBaking/Private/MaterialBakingModule.cpp
+++ b/Engine/Source/Developer/MaterialBaking/Private/MaterialBakingModule.cpp
@@ -456,7 +456,6 @@ public:
 		: MaterialBakingModule(InMaterialBakingModule)
 		, MaterialSettings(InMaterialSettings)
 		, MeshSettings(InMeshSettings)
-		, NumMaterials(MaterialSettings.Num())
 		, bSaveIntermediateTextures(CVarSaveIntermediateTextures.GetValueOnAnyThread() == 1)
 		, bEmissiveHDR(InMaterialBakingModule.bEmissiveHDR)
 	{
@@ -472,6 +471,8 @@ public:
 		}
 
 		ComputeMeshProcessingOrder();
+
+		NumMaterials = ProcessingOrder.Num();
 	}
 
 	virtual ~FMaterialBakingProcessor() = default;
@@ -492,46 +493,45 @@ public:
 			const int32 MaterialIndex = ProcessingOrder[Index];
 
 			const FMaterialDataEx& CurrentMaterialSettings = *MaterialSettings[MaterialIndex];
-			if (!CurrentMaterialSettings.PropertySizes.IsEmpty())
-			{
-				const FMeshData* CurrentMeshSettings = MeshSettings[MaterialIndex];
-				FBakeOutputEx& CurrentOutput = GetBakeOutput(MaterialIndex);
-
-				TMap<FRenderItemKey, FMeshMaterialRenderItem*>* RenderItems = GetRenderItems(Index);
-				check(RenderItems && !RenderItems->IsEmpty());
+			check(!CurrentMaterialSettings.PropertySizes.IsEmpty());
 
-				// For each property
-				for (const auto& [Property, Size] : CurrentMaterialSettings.PropertySizes)
-				{
-					TRACE_CPUPROFILER_EVENT_SCOPE_TEXT(*Property.ToString())
+			const FMeshData* CurrentMeshSettings = MeshSettings[MaterialIndex];
+			FBakeOutputEx& CurrentOutput = GetBakeOutput(MaterialIndex);
 
-					FExportMaterialProxy* ExportMaterialProxy = MaterialBakingModule.CreateMaterialProxy(&CurrentMaterialSettings, Property);
-					if (!ExportMaterialProxy->IsCompilationFinished())
-					{
-						TRACE_CPUPROFILER_EVENT_SCOPE(WaitForMaterialProxyCompilation)
-						ExportMaterialProxy->FinishCompilation();
-					}
+			TMap<FRenderItemKey, FMeshMaterialRenderItem*>* RenderItems = GetRenderItems(Index);
+			check(RenderItems && !RenderItems->IsEmpty());
 
-					UTextureRenderTarget2D* RenderTarget = GetRenderTarget(Property, Size, CurrentMaterialSettings);
-					FMeshMaterialRenderItem* RenderItem = RenderItems->FindChecked(FRenderItemKey(CurrentMeshSettings, Size));
+			// For each property
+			for (const auto& [Property, Size] : CurrentMaterialSettings.PropertySizes)
+			{
+				TRACE_CPUPROFILER_EVENT_SCOPE_TEXT(*Property.ToString())
 
-					BakeMaterialProperty(CurrentMaterialSettings, Property, RenderItem, RenderTarget, ExportMaterialProxy, CurrentOutput);
+				FExportMaterialProxy* ExportMaterialProxy = MaterialBakingModule.CreateMaterialProxy(&CurrentMaterialSettings, Property);
+				if (!ExportMaterialProxy->IsCompilationFinished())
+				{
+					TRACE_CPUPROFILER_EVENT_SCOPE(WaitForMaterialProxyCompilation)
+					ExportMaterialProxy->FinishCompilation();
 				}
 
-				// Destroying Render Items
-				// Must happen on the render thread to ensure they are not used anymore.
-				ENQUEUE_RENDER_COMMAND(DestroyRenderItems)(
-					[RenderItems](FRHICommandListImmediate& RHICmdList)
-					{
-						for (auto RenderItem : (*RenderItems))
-						{
-							delete RenderItem.Value;
-						}
+				UTextureRenderTarget2D* RenderTarget = GetRenderTarget(Property, Size, CurrentMaterialSettings);
+				FMeshMaterialRenderItem* RenderItem = RenderItems->FindChecked(FRenderItemKey(CurrentMeshSettings, Size));
 
-						delete RenderItems;
-					}
-				);
+				BakeMaterialProperty(CurrentMaterialSettings, Property, RenderItem, RenderTarget, ExportMaterialProxy, CurrentOutput);
 			}
+
+			// Destroying Render Items
+			// Must happen on the render thread to ensure they are not used anymore.
+			ENQUEUE_RENDER_COMMAND(DestroyRenderItems)(
+				[RenderItems](FRHICommandListImmediate& RHICmdList)
+				{
+					for (auto RenderItem : (*RenderItems))
+					{
+						delete RenderItem.Value;
+					}
+
+					delete RenderItems;
+				}
+			);
 		}
 	}
 
@@ -585,7 +585,10 @@ private:
 		ProcessingOrder.Reserve(MeshSettings.Num());
 		for (int32 Index = 0; Index < MeshSettings.Num(); ++Index)
 		{
-			ProcessingOrder.Add(Index);
+			if (!MaterialSettings[Index]->PropertySizes.IsEmpty())
+			{
+				ProcessingOrder.Add(Index);
+			}
 		}
 
 		// Start with the biggest mesh first so we can always reuse the same vertex/index buffers.
@@ -608,6 +611,8 @@ private:
 		const FMaterialDataEx* CurrentMaterialSettings = MaterialSettings[MaterialIndex];
 		const FMeshData* CurrentMeshSettings = MeshSettings[MaterialIndex];
 
+		check(!CurrentMaterialSettings->PropertySizes.IsEmpty());
+
 		for (const auto& [Property, Size] : CurrentMaterialSettings->PropertySizes)
 		{
 			FRenderItemKey RenderItemKey(CurrentMeshSettings, Size);
@@ -758,7 +763,6 @@ protected:
 
 	const TArray<FMaterialDataEx*>& MaterialSettings;
 	const TArray<FMeshData*>& MeshSettings;
-	const int32 NumMaterials;
 	const bool bSaveIntermediateTextures;
 	const bool bEmissiveHDR;
 
@@ -767,6 +771,7 @@ protected:
 	static const int32 PipelineDepth = 16;
 
 private:
+	int32 NumMaterials;
 	TArray<uint32> ProcessingOrder;
 
 	FMaterialBakingDynamicMeshBufferAllocator MaterialBakingDynamicMeshBufferAllocator;
