Material Baking - Fix crashes occurring when a GC was performed during material baking

diff --git a/Engine/Source/Developer/MaterialBaking/Private/MaterialBakingModule.cpp b/Engine/Source/Developer/MaterialBaking/Private/MaterialBakingModule.cpp
--- a/Engine/Source/Developer/MaterialBaking/Private/MaterialBakingModule.cpp
+++ b/Engine/Source/Developer/MaterialBaking/Private/MaterialBakingModule.cpp
@@ -280,6 +280,8 @@ void FMaterialBakingModule::StartupModule()
 {
 	bEmissiveHDR = false;
 
+	bIsBakingMaterials = false;
+
 	// Set which properties should enforce gamma correction
 	SetLinearBake(true);
 
@@ -816,7 +818,7 @@ public:
 				const bool bUsePooledRenderTargets = false;
 
 				UTextureRenderTarget2D* RenderTarget = CreateRenderTarget(Property, Size, bUsePooledRenderTargets, DefaultMaterialData->BackgroundColor);
-				RenderTargets.Add(Property, RenderTarget);
+				RenderTargets.Emplace(Property, RenderTarget);
 			}
 		}
 	}
@@ -872,11 +874,11 @@ public:
 
 	virtual UTextureRenderTarget2D* GetRenderTarget(FMaterialPropertyEx InMaterialProperty, const FIntPoint& InRequiredSize, const FMaterialDataEx& InMaterialSettings) override
 	{
-		return RenderTargets[InMaterialProperty];
+		return RenderTargets[InMaterialProperty].Get();
 	}
 
 private:
-	TMap<FMaterialPropertyEx, UTextureRenderTarget2D*> RenderTargets;
+	TMap<FMaterialPropertyEx, TStrongObjectPtr<UTextureRenderTarget2D>> RenderTargets;
 	FBakeOutputEx& Output;
 };
 
@@ -1073,6 +1075,8 @@ void FMaterialBakingModule::BakeMaterials(const TArray<FMaterialDataEx*>& Materi
 {
 	TRACE_CPUPROFILER_EVENT_SCOPE(FMaterialBakingModule::BakeMaterials)
 
+	TGuardValue<bool> GuardIsBaking(bIsBakingMaterials, true);
+
 	FMaterialBakingProcessorSingleOutput MaterialBakingProcessor(*this, MaterialSettings, MeshSettings, Output);
 	MaterialBakingProcessor.BakeMaterials();
 }
@@ -1082,6 +1086,8 @@ void FMaterialBakingModule::BakeMaterials(const TArray<FMaterialDataEx*>& Materi
 {
 	TRACE_CPUPROFILER_EVENT_SCOPE(FMaterialBakingModule::BakeMaterials)
 
+	TGuardValue<bool> GuardIsBaking(bIsBakingMaterials, true);
+
 	FMaterialBakingProcessorMultiOutput MaterialBakingProcessor(*this, MaterialSettings, MeshSettings, Output);
 	MaterialBakingProcessor.BakeMaterials();
 }
@@ -1164,10 +1170,6 @@ void FMaterialBakingModule::CleanupMaterialProxies()
 
 void FMaterialBakingModule::CleanupRenderTargets()
 {
-	for (auto RenderTarget : RenderTargetPool)
-	{
-		RenderTarget->RemoveFromRoot();
-	}
 	RenderTargetPool.Empty();
 }
 
@@ -1193,16 +1195,16 @@ UTextureRenderTarget2D* FMaterialBakingModule::CreateRenderTarget(FMaterialPrope
 	// First, look in pool
 	if (bInUsePooledRenderTargets)
 	{
-		auto RenderTargetComparison = [bForceLinearGamma, PixelFormat, ClampedTargetSize](const UTextureRenderTarget2D* CompareRenderTarget) -> bool
+		auto RenderTargetComparison = [bForceLinearGamma, PixelFormat, ClampedTargetSize](const TStrongObjectPtr<UTextureRenderTarget2D>& CompareRenderTarget) -> bool
 		{
 			return (CompareRenderTarget->SizeX == ClampedTargetSize.X && CompareRenderTarget->SizeY == ClampedTargetSize.Y && CompareRenderTarget->OverrideFormat == PixelFormat && CompareRenderTarget->bForceLinearGamma == bForceLinearGamma);
 		};
 
 		// Find any pooled render target with suitable properties.
-		UTextureRenderTarget2D** FindResult = RenderTargetPool.FindByPredicate(RenderTargetComparison);
+		TStrongObjectPtr<UTextureRenderTarget2D>* FindResult = RenderTargetPool.FindByPredicate(RenderTargetComparison);
 		if (FindResult)
 		{
-			RenderTarget = *FindResult;
+			RenderTarget = FindResult->Get();
 		}
 	}
 
@@ -1221,8 +1223,7 @@ UTextureRenderTarget2D* FMaterialBakingModule::CreateRenderTarget(FMaterialPrope
 
 		if (bInUsePooledRenderTargets)
 		{
-			RenderTarget->AddToRoot();
-			RenderTargetPool.Add(RenderTarget);
+			RenderTargetPool.Emplace(RenderTarget);
 		}
 	}
 
@@ -1413,7 +1414,11 @@ void FMaterialBakingModule::OnObjectModified(UObject* Object)
 
 void FMaterialBakingModule::OnPreGarbageCollect()
 {
-	CleanupMaterialProxies();
+	// Do not cleanup material proxies while baking materials.
+	if (!bIsBakingMaterials)
+	{
+		CleanupMaterialProxies();
+	}
 }
 
 #undef LOCTEXT_NAMESPACE //"MaterialBakingModule"
diff --git a/Engine/Source/Developer/MaterialBaking/Private/MaterialBakingModule.h b/Engine/Source/Developer/MaterialBaking/Private/MaterialBakingModule.h
--- a/Engine/Source/Developer/MaterialBaking/Private/MaterialBakingModule.h
+++ b/Engine/Source/Developer/MaterialBaking/Private/MaterialBakingModule.h
@@ -2,7 +2,7 @@
 
 #pragma once
 
-#include "CoreMinimal.h"
+#include "UObject/StrongObjectPtr.h"
 #include "Modules/ModuleInterface.h"
 #include "PixelFormat.h"
 #include "SceneTypes.h"
@@ -84,7 +84,7 @@ private:
 	};
 
 	/** Pool of available render targets, cached for re-using on consecutive property rendering */
-	TArray<UTextureRenderTarget2D*> RenderTargetPool;
+	TArray<TStrongObjectPtr<UTextureRenderTarget2D>> RenderTargetPool;
 
 	/** Pool of cached material proxies to optimize material baking workflow, stays resident when MaterialBaking.UseMaterialProxyCaching is set to 1 */
 	typedef TWeakObjectPtr<UMaterialInterface>				FMaterialPoolKey;
@@ -101,4 +101,6 @@ private:
 	EPropertyColorSpace DefaultColorSpace;
 
 	bool bEmissiveHDR;
+
+	bool bIsBakingMaterials;
 };
