Fix clang 18 -Wshadow warnings

#rnx

diff --git a/Engine/Plugins/Experimental/Chooser/Source/ChooserEditor/Private/SChooserTableRow.cpp b/Engine/Plugins/Experimental/Chooser/Source/ChooserEditor/Private/SChooserTableRow.cpp
--- a/Engine/Plugins/Experimental/Chooser/Source/ChooserEditor/Private/SChooserTableRow.cpp
+++ b/Engine/Plugins/Experimental/Chooser/Source/ChooserEditor/Private/SChooserTableRow.cpp
@@ -114,13 +114,13 @@ namespace UE::ChooserEditor
 			{
 				UChooserTable* ContextOwner = Chooser->GetContextOwner();
 				TSharedPtr<SWidget> ResultWidget = FObjectChooserWidgetFactories::CreateWidget(false, ContextOwner, FObjectChooserBase::StaticStruct(), Chooser->ResultsStructs[RowIndex->RowIndex].GetMutableMemory(), Chooser->ResultsStructs[RowIndex->RowIndex].GetScriptStruct(), ContextOwner->OutputObjectType,
-				FOnStructPicked::CreateLambda([this, RowIndex=RowIndex->RowIndex](const UScriptStruct* ChosenStruct)
+				FOnStructPicked::CreateLambda([this, InRowIndex=RowIndex->RowIndex](const UScriptStruct* ChosenStruct)
 				{
 					UChooserTable* ContextOwner = Chooser->GetContextOwner();
 					const FScopedTransaction Transaction(LOCTEXT("Change Row Result Type", "Change Row Result Type"));
 					Chooser->Modify(true);
-					Chooser->ResultsStructs[RowIndex].InitializeAs(ChosenStruct);
-					FObjectChooserWidgetFactories::CreateWidget(false, ContextOwner, FObjectChooserBase::StaticStruct(), Chooser->ResultsStructs[RowIndex].GetMutableMemory(), ChosenStruct, ContextOwner->OutputObjectType, FOnStructPicked(), &CacheBorder);
+					Chooser->ResultsStructs[InRowIndex].InitializeAs(ChosenStruct);
+					FObjectChooserWidgetFactories::CreateWidget(false, ContextOwner, FObjectChooserBase::StaticStruct(), Chooser->ResultsStructs[InRowIndex].GetMutableMemory(), ChosenStruct, ContextOwner->OutputObjectType, FOnStructPicked(), &CacheBorder);
 				}),
 				&CacheBorder
 				);
diff --git a/Engine/Plugins/Experimental/TextureGraph/Source/TextureGraphEditor/Private/EdGraph/TG_EdGraphNode.cpp b/Engine/Plugins/Experimental/TextureGraph/Source/TextureGraphEditor/Private/EdGraph/TG_EdGraphNode.cpp
--- a/Engine/Plugins/Experimental/TextureGraph/Source/TextureGraphEditor/Private/EdGraph/TG_EdGraphNode.cpp
+++ b/Engine/Plugins/Experimental/TextureGraph/Source/TextureGraphEditor/Private/EdGraph/TG_EdGraphNode.cpp
@@ -111,13 +111,13 @@ void UTG_EdGraphNode::GetNodeContextMenuActions(UToolMenu* Menu, class UGraphNod
 			LOCTEXT("SelectPinTooltip", "Set this pin as thumb and preview"),
 			FSlateIcon(),
 			FUIAction(
-				FExecuteAction::CreateLambda([Pin = Context->Pin, Node = Context->Node, this]
+				FExecuteAction::CreateLambda([Pin = Context->Pin, InNode = Context->Node, this]
 				{
-					const UTG_EdGraphNode* TGEdNode = Cast<UTG_EdGraphNode>(Node);
+					const UTG_EdGraphNode* TGEdNode = Cast<UTG_EdGraphNode>(InNode);
 					if (TGEdNode)
 					{
 						const_cast<UTG_EdGraphNode*>(TGEdNode)->SelectPin(const_cast<UEdGraphPin*>(Pin), true); // assign the current selected pin ion the node
-						UTG_EdGraph* TSEdGraph = Cast<UTG_EdGraph>(Node->GetGraph());
+						UTG_EdGraph* TSEdGraph = Cast<UTG_EdGraph>(InNode->GetGraph());
 						TSEdGraph->PinSelectionManager.UpdateSelection(const_cast<UEdGraphPin*>(Pin));
 					}
 				}),
diff --git a/Engine/Source/Developer/MaterialBaking/Private/MaterialBakingModule.cpp b/Engine/Source/Developer/MaterialBaking/Private/MaterialBakingModule.cpp
--- a/Engine/Source/Developer/MaterialBaking/Private/MaterialBakingModule.cpp
+++ b/Engine/Source/Developer/MaterialBaking/Private/MaterialBakingModule.cpp
@@ -635,9 +635,9 @@ private:
 			PreparedRenderItems[NextRenderItem % PipelineDepth] =
 				Async(
 					EAsyncExecution::ThreadPool,
-					[this, NextRenderItem=NextRenderItem]()
+					[this, Index=NextRenderItem]()
 					{
-						return PrepareRenderItems_AnyThread(ProcessingOrder[NextRenderItem]);
+						return PrepareRenderItems_AnyThread(ProcessingOrder[Index]);
 					}
 			);
 		}
@@ -908,12 +908,12 @@ public:
 	~FMaterialBakingProcessorMultiOutput()
 	{
 		ENQUEUE_RENDER_COMMAND(ProcessRemainingReads)(
-		[this, PipelineIndex=PipelineIndex](FRHICommandListImmediate& RHICmdList)
+		[this, InPipelineIndex=PipelineIndex](FRHICommandListImmediate& RHICmdList)
 		{
 			// Enqueue remaining reads
 			for (int32 Index = 0; Index < PipelineDepth; Index++)
 			{
-				int32 LocalPipelineIndex = (PipelineIndex + Index) % PipelineDepth;
+				int32 LocalPipelineIndex = (InPipelineIndex + Index) % PipelineDepth;
 
 				if (PipelineContext[LocalPipelineIndex].ReadCommand)
 				{
@@ -949,7 +949,7 @@ private:
 	virtual void OnMaterialPropertyBaked(const FMaterialDataEx& CurrentMaterialSettings, const FMaterialPropertyEx& Property, FMeshMaterialRenderItem* RenderItem, UTextureRenderTarget2D* RenderTarget, FExportMaterialProxy* ExportMaterialProxy, FBakeOutputEx& CurrentOutput) override
 	{
 		ENQUEUE_RENDER_COMMAND(CopyStagingBuffer)(
-			[this, RenderTarget, Property, ExportMaterialProxy, &CurrentMaterialSettings, &CurrentOutput, PipelineIndex=PipelineIndex](FRHICommandListImmediate& RHICmdList)
+			[this, RenderTarget, Property, ExportMaterialProxy, &CurrentMaterialSettings, &CurrentOutput, InPipelineIndex=PipelineIndex](FRHICommandListImmediate& RHICmdList)
 			{
 				FTextureRenderTargetResource* RenderTargetResource = RenderTarget->GetRenderTargetResource();
 
@@ -1010,13 +1010,13 @@ private:
 					};
 
 				// Run previous command if we're going to overwrite it meaning pipeline depth has been reached
-				if (PipelineContext[PipelineIndex].ReadCommand)
+				if (PipelineContext[InPipelineIndex].ReadCommand)
 				{
-					PipelineContext[PipelineIndex].ReadCommand(RHICmdList);
+					PipelineContext[InPipelineIndex].ReadCommand(RHICmdList);
 				}
 
 				// Generate a texture reading command that will be executed once it reaches the end of the pipeline
-				PipelineContext[PipelineIndex].ReadCommand =
+				PipelineContext[InPipelineIndex].ReadCommand =
 					[FinalProcessing_AnyThread, StagingBufferRef = MoveTemp(StagingBufferRef), GPUFence = MoveTemp(GPUFence)](FRHICommandListImmediate& RHICmdList) mutable
 					{
 						TRACE_CPUPROFILER_EVENT_SCOPE(MapAndEnqueue)
diff --git a/Engine/Source/Editor/Kismet/Private/DiffControl.cpp b/Engine/Source/Editor/Kismet/Private/DiffControl.cpp
--- a/Engine/Source/Editor/Kismet/Private/DiffControl.cpp
+++ b/Engine/Source/Editor/Kismet/Private/DiffControl.cpp
@@ -351,9 +351,9 @@ void FDetailsDiffControl::GenerateTreeEntriesWithoutComments(
 	for (auto&[Object, DetailsDiff] : DetailsDiffs)
 	{
 		Algo::Transform(DifferingProperties, PropertyAllowList,
-        [&Object = Object](const FSingleObjectDiffEntry& DiffEntry)
+        [&InObject = Object](const FSingleObjectDiffEntry& DiffEntry)
         {
-        	return DiffEntry.Identifier.ResolvePath(Object);
+        	return DiffEntry.Identifier.ResolvePath(InObject);
         });
         
         DetailsDiff.DetailsWidget()->UpdatePropertyAllowList(PropertyAllowList);
