Fix for material exports attempting to write to depth textures when none are bound

diff --git a/Engine/Source/Developer/MaterialBaking/Private/ExportMaterialProxy.h b/Engine/Source/Developer/MaterialBaking/Private/ExportMaterialProxy.h
--- a/Engine/Source/Developer/MaterialBaking/Private/ExportMaterialProxy.h
+++ b/Engine/Source/Developer/MaterialBaking/Private/ExportMaterialProxy.h
@@ -211,7 +211,7 @@ struct FExportMaterialCompiler : public FProxyMaterialCompiler
 class FExportMaterialProxy : public FMaterial, public FMaterialRenderProxy
 {
 public:
-	FExportMaterialProxy(UMaterialInterface* InMaterialInterface, EMaterialProperty InPropertyToCompile, const FString& InCustomOutputToCompile = TEXT(""), bool bInSynchronousCompilation = true, bool bTangentSpaceNormal = false, EBlendMode ProxyBlendMode = BLEND_Opaque)
+	FExportMaterialProxy(UMaterialInterface* InMaterialInterface, EMaterialProperty InPropertyToCompile, const FString& InCustomOutputToCompile = TEXT(""), bool bInSynchronousCompilation = true, bool bTangentSpaceNormal = false, EBlendMode ProxyBlendMode = BLEND_Opaque, bool bAllowPixelDepthOffset = true)
 		: FMaterial()
 		, FMaterialRenderProxy(GetPathNameSafe(InMaterialInterface->GetMaterial()))
 		, MaterialInterface(InMaterialInterface)
@@ -273,6 +273,8 @@ public:
 		Usage = ResourceId.Usage;
 		ResourceId.BaseMaterialId = Material->StateId;
 
+		SetAllowPixelDepthOffset(bAllowPixelDepthOffset);
+		
 		CacheShaders(ResourceId, GMaxRHIShaderPlatform);
 	}
 
diff --git a/Engine/Source/Developer/MaterialBaking/Private/MaterialBakingModule.cpp b/Engine/Source/Developer/MaterialBaking/Private/MaterialBakingModule.cpp
--- a/Engine/Source/Developer/MaterialBaking/Private/MaterialBakingModule.cpp
+++ b/Engine/Source/Developer/MaterialBaking/Private/MaterialBakingModule.cpp
@@ -1257,7 +1257,7 @@ FExportMaterialProxy* FMaterialBakingModule::CreateMaterialProxy(const FMaterial
 	// Not found, create a new entry
 	if (Proxy == nullptr)
 	{
-		Proxy = new FExportMaterialProxy(MaterialSettings->Material, Property.Type, Property.CustomOutput.ToString(), false /* bInSynchronousCompilation */, MaterialSettings->bTangentSpaceNormal, MaterialSettings->BlendMode);
+		Proxy = new FExportMaterialProxy(MaterialSettings->Material, Property.Type, Property.CustomOutput.ToString(), false /* bInSynchronousCompilation */, MaterialSettings->bTangentSpaceNormal, MaterialSettings->BlendMode, false);
 		MaterialProxyPool.Add(MaterialSettings->Material, FMaterialPoolValue(Property, Proxy));
 	}
 
diff --git a/Engine/Source/Runtime/Engine/Private/Materials/HLSLMaterialTranslator.cpp b/Engine/Source/Runtime/Engine/Private/Materials/HLSLMaterialTranslator.cpp
--- a/Engine/Source/Runtime/Engine/Private/Materials/HLSLMaterialTranslator.cpp
+++ b/Engine/Source/Runtime/Engine/Private/Materials/HLSLMaterialTranslator.cpp
@@ -1785,7 +1785,7 @@ void FHLSLMaterialTranslator::TranslateMaterial()
 	bAllowCodeChunkGeneration = false;
 
 	bUsesEmissiveColor = bSubstrateWritesEmissive || IsMaterialPropertyUsed(MP_EmissiveColor, Chunk[MP_EmissiveColor], FLinearColor(0, 0, 0, 0), 3);
-	bUsesPixelDepthOffset = (AllowPixelDepthOffset(Platform) && IsMaterialPropertyUsed(MP_PixelDepthOffset, Chunk[MP_PixelDepthOffset], FLinearColor(0, 0, 0, 0), 1));
+	bUsesPixelDepthOffset = (AllowPixelDepthOffset(Platform) && Material->AllowPixelDepthOffset() && IsMaterialPropertyUsed(MP_PixelDepthOffset, Chunk[MP_PixelDepthOffset], FLinearColor(0, 0, 0, 0), 1));
 
 	bool bUsesWorldPositionOffsetCurrent = IsMaterialPropertyUsed(MP_WorldPositionOffset, Chunk[MP_WorldPositionOffset], FLinearColor(0, 0, 0, 0), 3);
 	bool bUsesWorldPositionOffsetPrevious = IsMaterialPropertyUsed(MP_WorldPositionOffset, Chunk[CompiledMP_PrevWorldPositionOffset], FLinearColor(0, 0, 0, 0), 3);
diff --git a/Engine/Source/Runtime/Engine/Public/MaterialShared.h b/Engine/Source/Runtime/Engine/Public/MaterialShared.h
--- a/Engine/Source/Runtime/Engine/Public/MaterialShared.h
+++ b/Engine/Source/Runtime/Engine/Public/MaterialShared.h
@@ -1998,7 +1998,8 @@ public:
 		bContainsInlineShaders(false),
 		bLoadedCookedShaderMapId(false),
 		bGameThreadShaderMapIsComplete(false),
-		bRenderingThreadShaderMapIsComplete(false)
+		bRenderingThreadShaderMapIsComplete(false),
+		bAllowPixelDepthOffset(true)
 	{
 		// this option affects only deferred renderer
 		static TConsoleVariableData<int32>* CVarStencilDitheredLOD;
@@ -2690,6 +2691,9 @@ protected:
 	virtual UMaterialExpression* GetMaterialGraphNodePreviewExpression() const { return nullptr; }
 #endif
 
+	void SetAllowPixelDepthOffset(bool bEnable) { bAllowPixelDepthOffset = bEnable; }
+	bool AllowPixelDepthOffset() const { return bAllowPixelDepthOffset; }
+	
 private:
 #if WITH_EDITOR
 	bool ReleaseGameThreadCompilingShaderMap();
@@ -2781,6 +2785,8 @@ private:
 
 	uint32 bGameThreadShaderMapIsComplete : 1;
 	uint32 bRenderingThreadShaderMapIsComplete : 1;
+	
+	uint32 bAllowPixelDepthOffset : 1;
 
 #if UE_CHECK_FMATERIAL_LIFETIME
 	/** Set when the owner of this FMaterial (typically a UMaterial or UMaterialInstance) has had BeginDestroy() called */
