Substrate - exclude FrontMaterial from baking options when Substrate is enabled or disabled.

#rnx
#tests PC editor

diff --git a/Engine/Source/Developer/MaterialBaking/Private/MaterialOptionsCustomization.cpp b/Engine/Source/Developer/MaterialBaking/Private/MaterialOptionsCustomization.cpp
--- a/Engine/Source/Developer/MaterialBaking/Private/MaterialOptionsCustomization.cpp
+++ b/Engine/Source/Developer/MaterialBaking/Private/MaterialOptionsCustomization.cpp
@@ -15,7 +15,7 @@
 #include "IPropertyTypeCustomization.h"
 #include "IPropertyUtilities.h"
 
-static const TArray<EMaterialProperty> DisabledProperties = { MP_Refraction };
+static const TArray<EMaterialProperty> DisabledProperties = { MP_Refraction, MP_FrontMaterial };
 
 TSharedRef<IPropertyTypeCustomization> FPropertyEntryCustomization::MakeInstance()
 {
