substrate - fix for compilation warning.

#rnx

diff --git a/Engine/Source/Developer/MaterialBaking/Private/ExportMaterialProxy.h b/Engine/Source/Developer/MaterialBaking/Private/ExportMaterialProxy.h
--- a/Engine/Source/Developer/MaterialBaking/Private/ExportMaterialProxy.h
+++ b/Engine/Source/Developer/MaterialBaking/Private/ExportMaterialProxy.h
@@ -24,6 +24,8 @@
 #include "Materials/MaterialAttributeDefinitionMap.h"
 #include "Materials/MaterialExpressionCustomOutput.h"
 
+#include "RenderUtils.h"
+
 struct FExportMaterialCompiler : public FProxyMaterialCompiler
 {
 	FExportMaterialCompiler(FMaterialCompiler* InCompiler) :
