Expose EMaterialProperty::MP_WorldPositionOffset so that it can be picked from the ConnectMaterialProperty BP node.
Hide it from the material baking UI since we won't be baking that parameter.

diff --git a/Engine/Source/Developer/MaterialBaking/Private/MaterialOptionsCustomization.cpp b/Engine/Source/Developer/MaterialBaking/Private/MaterialOptionsCustomization.cpp
--- a/Engine/Source/Developer/MaterialBaking/Private/MaterialOptionsCustomization.cpp
+++ b/Engine/Source/Developer/MaterialBaking/Private/MaterialOptionsCustomization.cpp
@@ -93,6 +93,8 @@ void FPropertyEntryCustomization::UpdateRestrictions(const int32 EntryIndex)
 	{
 		PropertyRestriction->AddDisabledValue(PropertyEnum->GetNameStringByValue(Property));
 	}
+
+	PropertyRestriction->AddHiddenValue(GET_ENUMERATOR_NAME_CHECKED(EMaterialProperty, MP_WorldPositionOffset).ToString());
 }
 
 TSharedRef<IDetailCustomization> FMaterialOptionsCustomization::MakeInstance(int32 InNumLODs)
diff --git a/Engine/Source/Runtime/Engine/Public/SceneTypes.h b/Engine/Source/Runtime/Engine/Public/SceneTypes.h
--- a/Engine/Source/Runtime/Engine/Public/SceneTypes.h
+++ b/Engine/Source/Runtime/Engine/Public/SceneTypes.h
@@ -158,7 +158,7 @@ enum EMaterialProperty : int
 	MP_Anisotropy UMETA(DisplayName = "Anisotropy"),
 	MP_Normal UMETA(DisplayName = "Normal"),
 	MP_Tangent UMETA(DisplayName = "Tangent"),
-	MP_WorldPositionOffset UMETA(Hidden),
+	MP_WorldPositionOffset UMETA(DisplayName = "World Position Offset"),
 	MP_WorldDisplacement_DEPRECATED UMETA(Hidden),
 	MP_TessellationMultiplier_DEPRECATED UMETA(Hidden),
 	MP_SubsurfaceColor UMETA(DisplayName = "Subsurface"),
