[glTF Exporter -> MaterialBaking] Fixing Emissive property baking.

#rnx

diff --git a/Engine/Source/Developer/MaterialBaking/Private/ExportMaterialProxy.h b/Engine/Source/Developer/MaterialBaking/Private/ExportMaterialProxy.h
--- a/Engine/Source/Developer/MaterialBaking/Private/ExportMaterialProxy.h
+++ b/Engine/Source/Developer/MaterialBaking/Private/ExportMaterialProxy.h
@@ -601,7 +601,7 @@ public:
 	{
 		// Substrate needs the real material shading model since the expressions access GetMaterialShadingModels() through the compiler to generate the substrate operators, 
 		// and we do not want unlit materials when it is a Slab or a ShadingModel node.
-		return Substrate::IsSubstrateEnabled() ? Material->GetShadingModels() : MSM_Unlit;
+		return Substrate::IsSubstrateEnabled() ? Material->GetShadingModels() : MSM_DefaultLit;
 	}
 
 	virtual bool IsShadingModelFromMaterialExpression() const override { return false; }
